Monitor network status and resume background download if reconnected

#android

diff --git a/Engine/Plugins/Online/Android/AndroidFetchBackgroundDownload/Source/AndroidFetchBackgroundDownload/AndroidFetchBackgroundDownload_UPL.xml b/Engine/Plugins/Online/Android/AndroidFetchBackgroundDownload/Source/AndroidFetchBackgroundDownload/AndroidFetchBackgroundDownload_UPL.xml
--- a/Engine/Plugins/Online/Android/AndroidFetchBackgroundDownload/Source/AndroidFetchBackgroundDownload/AndroidFetchBackgroundDownload_UPL.xml
+++ b/Engine/Plugins/Online/Android/AndroidFetchBackgroundDownload/Source/AndroidFetchBackgroundDownload/AndroidFetchBackgroundDownload_UPL.xml
@@ -47,4 +47,48 @@
 		<log text="Copying AndroidFetchBackgroundDownload files to staging"/>
 		<copyDir src="$S(PluginDir)/Java" dst="$S(BuildDir)/src/com/epicgames/unreal" />
 	</prebuildCopies>
+
+	<gameActivityImportAdditions>
+		<insert>
+			import com.epicgames.unreal.network.NetworkConnectivityClient;
+			import com.epicgames.unreal.network.NetworkChangedManager;
+		</insert>
+	</gameActivityImportAdditions>
+
+	<gameActivityClassAdditions>
+		<insert>
+			private NetworkConnectivityClient.Listener AFBD_NetworkListener = null;
+
+			public native void nativeAndroidBackgroundDownloadNetworkChanged();
+
+			public void AndroidThunkJava_AFBD_AddNetworkListener()
+			{
+				if (AFBD_NetworkListener == null)
+				{
+					AFBD_NetworkListener = new NetworkConnectivityClient.Listener() {
+						@Override
+						public void onNetworkAvailable(NetworkConnectivityClient.NetworkTransportType networkTransportType)
+						{
+							nativeAndroidBackgroundDownloadNetworkChanged();
+						}
+						@Override
+						public void onNetworkLost()
+						{
+							nativeAndroidBackgroundDownloadNetworkChanged();
+						}
+					};
+					NetworkChangedManager.getInstance().addListener(AFBD_NetworkListener);
+				}
+			}
+
+			public void AndroidThunkJava_AFBD_RemoveNetworkListener()
+			{
+				if (AFBD_NetworkListener != null)
+				{
+					NetworkChangedManager.getInstance().removeListener(AFBD_NetworkListener);
+					AFBD_NetworkListener = null;
+				}
+			}
+		</insert>
+	</gameActivityClassAdditions>
 </root>
diff --git a/Engine/Plugins/Online/Android/AndroidFetchBackgroundDownload/Source/AndroidFetchBackgroundDownload/Java/download/fetch/FetchManager.java b/Engine/Plugins/Online/Android/AndroidFetchBackgroundDownload/Source/AndroidFetchBackgroundDownload/Java/download/fetch/FetchManager.java
--- a/Engine/Plugins/Online/Android/AndroidFetchBackgroundDownload/Source/AndroidFetchBackgroundDownload/Java/download/fetch/FetchManager.java
+++ b/Engine/Plugins/Online/Android/AndroidFetchBackgroundDownload/Source/AndroidFetchBackgroundDownload/Java/download/fetch/FetchManager.java
@@ -26,6 +26,7 @@ import com.tonyodev.fetch2.NetworkType;
 import com.tonyodev.fetch2.Priority;
 import com.tonyodev.fetch2.Request;
 import com.tonyodev.fetch2.exception.FetchException;
+import com.tonyodev.fetch2core.FetchLogger;
 import com.tonyodev.fetch2core.Func;
 import com.tonyodev.fetch2core.Func2;
 import com.tonyodev.fetch2.Status;
@@ -358,9 +359,13 @@ public class FetchManager implements FetchDownloadProgressOwner, FetchEnqueueRes
 
 		if (FetchInstance == null)
 		{
+//			InstanceLogger = new FetchLogger(true, "FManager");
+
 			//TODO TRoss: Pull these values from the worker's getInputData
 			FetchInstance = Fetch.Impl.getInstance(new FetchConfiguration.Builder(context)
 				.setNamespace(context.getPackageName())
+//				.enableLogging(true)
+//				.setLogger(InstanceLogger)
 				.enableRetryOnNetworkGain(true)
 				.setProgressReportingInterval(200)
 				.build());
@@ -949,6 +954,7 @@ public class FetchManager implements FetchDownloadProgressOwner, FetchEnqueueRes
 	//private boolean ShouldRetryRequest(@NonNull Request 
 	
 	private Fetch FetchInstance = null;
+	private FetchLogger InstanceLogger = null;
 
 	private FetchRequestProgressListener FetchListener = null;
 	
diff --git a/Engine/Plugins/Online/Android/AndroidFetchBackgroundDownload/Source/AndroidFetchBackgroundDownload/Private/AndroidPlatformBackgroundHttpManager.cpp b/Engine/Plugins/Online/Android/AndroidFetchBackgroundDownload/Source/AndroidFetchBackgroundDownload/Private/AndroidPlatformBackgroundHttpManager.cpp
--- a/Engine/Plugins/Online/Android/AndroidFetchBackgroundDownload/Source/AndroidFetchBackgroundDownload/Private/AndroidPlatformBackgroundHttpManager.cpp
+++ b/Engine/Plugins/Online/Android/AndroidFetchBackgroundDownload/Source/AndroidFetchBackgroundDownload/Private/AndroidPlatformBackgroundHttpManager.cpp
@@ -78,6 +78,22 @@ const FString FAndroidNativeDownloadWorkerParameterKeys::NOTIFICATION_RESOURCE_S
 
 FAndroidPlatformBackgroundHttpManager::FJavaClassInfo FAndroidPlatformBackgroundHttpManager::JavaInfo = FAndroidPlatformBackgroundHttpManager::FJavaClassInfo();
 
+static bool bNetworkConnected = false;
+static volatile int32 bNetworkReconnected = false;
+
+static bool GetNetworkConnected()
+{
+	switch (FPlatformMisc::GetNetworkConnectionType())
+	{
+		case ENetworkConnectionType::Ethernet:
+		case ENetworkConnectionType::Cell:
+		case ENetworkConnectionType::WiFi:
+		case ENetworkConnectionType::WiMAX:
+			return true;
+	}
+	return false;
+}
+
 FAndroidPlatformBackgroundHttpManager::FAndroidPlatformBackgroundHttpManager()
 	: bHasPendingCompletes(false)
 	, bIsModifyingPauseList(false)
@@ -85,6 +101,17 @@ FAndroidPlatformBackgroundHttpManager::FAndroidPlatformBackgroundHttpManager()
 	, bIsModifyingCancelList(false)
 	, AndroidBackgroundHTTPManagerDefaultLocalizedText()
 {
+	if (JNIEnv* Env = FAndroidApplication::GetJavaEnv())
+	{
+		static jmethodID AddNetworkListenerFunc = FJavaWrapper::FindMethod(Env, FJavaWrapper::GameActivityClassID, "AndroidThunkJava_AFBD_AddNetworkListener", "()V", false);
+		if (AddNetworkListenerFunc != nullptr)
+		{
+			FJavaWrapper::CallVoidMethod(Env, FJavaWrapper::GameActivityThis, AddNetworkListenerFunc);
+		}
+	}
+
+	bNetworkReconnected = false;
+	bNetworkConnected = GetNetworkConnected();
 }
 
 bool FAndroidPlatformBackgroundHttpManager::HandleRequirementsCheck()
@@ -152,6 +179,17 @@ bool FAndroidPlatformBackgroundHttpManager::Tick(float DeltaTime)
 void FAndroidPlatformBackgroundHttpManager::ActivatePendingRequests()
 {
 	UE_LOG(LogBackgroundHttpManager, VeryVerbose, TEXT("ActivatePendingRequests called"));
+
+	bool bRestoreActiveRequests = false;
+	if (FPlatformAtomics::AtomicRead(&bNetworkReconnected))
+	{
+		FRWScopeLock ScopeLock(ActiveRequestLock, SLT_ReadOnly);
+		int32 ActiveRequestCount = ActiveRequests.Num();
+		UE_LOG(LogBackgroundHttpManager, Display, TEXT("Reconnected and found %d active requests to restore"), ActiveRequestCount);
+		bRestoreActiveRequests = (ActiveRequestCount > 0);
+		FPlatformAtomics::InterlockedExchange(&bNetworkReconnected, false);
+	}
+
 	bool bHasPendingRequests = false;
 	{
 		FRWScopeLock ScopeLock(PendingRequestLock, SLT_ReadOnly);
@@ -160,9 +198,10 @@ void FAndroidPlatformBackgroundHttpManager::ActivatePendingRequests()
 		UE_LOG(LogBackgroundHttpManager, VeryVerbose, TEXT("bHasPendingRequests: %d"), static_cast<int>(bHasPendingRequests));
 	}
 		
-	if (bHasPendingRequests)
+	if (bHasPendingRequests || bRestoreActiveRequests)
 	{
-		UE_LOG(LogBackgroundHttpManager, Display, TEXT("Found PendingRequests to activate"));
+		UE_CLOG(bHasPendingRequests, LogBackgroundHttpManager, Display, TEXT("Found PendingRequests to activate"));
+		UE_CLOG(bRestoreActiveRequests, LogBackgroundHttpManager, Display, TEXT("Found ActiveRequests to restore"));
 
 		const bool bIsOptional = false;
 
@@ -1064,5 +1103,31 @@ JNI_METHOD void Java_com_epicgames_unreal_download_UEDownloadWorker_nativeAndroi
 	FAndroidBackgroundDownloadDelegates::AndroidBackgroundDownload_OnTickWorkerThread.Broadcast(jenv, thiz);
 }
 
+//
+//JNI methods coming from GameActivity
+//
+
+JNI_METHOD void Java_com_epicgames_unreal_GameActivity_nativeAndroidBackgroundDownloadNetworkChanged(JNIEnv* jenv, jobject thiz)
+{
+	bool bNewNetworkConnected = GetNetworkConnected();
+
+	if (bNewNetworkConnected != bNetworkConnected)
+	{
+		bNetworkConnected = bNewNetworkConnected;
+
+		if (bNetworkConnected)
+		{
+			// reconnected !
+			//FPlatformMisc::LowLevelOutputDebugString(TEXT("nativeAFBD_NetworkChanged: Reconnected!"));
+			FPlatformAtomics::InterlockedExchange(&bNetworkReconnected, true);
+		}
+		else
+		{
+			// disconnected !
+			//FPlatformMisc::LowLevelOutputDebugString(TEXT("nativeAFBD_NetworkChanged: Disconnected!"));
+		}
+	}
+}
+
 
 #undef LOCTEXT_NAMESPACE
\ No newline at end of file
