Avoid suspended download causing faulty notification values

diff --git a/Engine/Plugins/Online/Android/AndroidFetchBackgroundDownload/Source/AndroidFetchBackgroundDownload/Java/download/fetch/FetchManager.java b/Engine/Plugins/Online/Android/AndroidFetchBackgroundDownload/Source/AndroidFetchBackgroundDownload/Java/download/fetch/FetchManager.java
--- a/Engine/Plugins/Online/Android/AndroidFetchBackgroundDownload/Source/AndroidFetchBackgroundDownload/Java/download/fetch/FetchManager.java
+++ b/Engine/Plugins/Online/Android/AndroidFetchBackgroundDownload/Source/AndroidFetchBackgroundDownload/Java/download/fetch/FetchManager.java
@@ -830,7 +830,10 @@ public class FetchManager implements FetchDownloadProgressOwner, FetchEnqueueRes
 		EDownloadCompleteReason CompleteReasonToSend = ConvertCompleteReasonForDownload(CompleteReason);
 		
 		//Mark download as complete
-		DownloadDesc.PreviousDownloadPercent = 100;
+		if (CompleteReason == ECompleteReason.Success)
+		{
+			DownloadDesc.PreviousDownloadPercent = 100;
+		}
 		CompletedDownloads.put(DownloadDesc.RequestID, DownloadDesc);
 
 		//Only bubble up to the UEDownloadWorker non-intenional completes as we initiated any other completes internally
