Keep the Download notification alive when the Network is disconnected.
Caching the last valid value and repeat it while the network is down

diff --git a/Engine/Plugins/Online/Android/AndroidFetchBackgroundDownload/Config/BaseAndroidFetchBackgroundDownload.ini b/Engine/Plugins/Online/Android/AndroidFetchBackgroundDownload/Config/BaseAndroidFetchBackgroundDownload.ini
--- a/Engine/Plugins/Online/Android/AndroidFetchBackgroundDownload/Config/BaseAndroidFetchBackgroundDownload.ini
+++ b/Engine/Plugins/Online/Android/AndroidFetchBackgroundDownload/Config/BaseAndroidFetchBackgroundDownload.ini
@@ -68,4 +68,21 @@
 +DefaultNotificationText_Cancel=("ru","Отменить")
 +DefaultNotificationText_Cancel=("tr","İptal")
 +DefaultNotificationText_Cancel=("zh-CN","取消")
-+DefaultNotificationText_Cancel=("zh-Hant","取消")
\ No newline at end of file
++DefaultNotificationText_Cancel=("zh-Hant","取消")
+
+;NoInternet Text
++DefaultNotificationText_NoInternet=("native","No Internet Available")
++DefaultNotificationText_NoInternet=("ar","No Internet Available")
++DefaultNotificationText_NoInternet=("de","Kein Internetzugang")
++DefaultNotificationText_NoInternet=("es","No Internet Available")
++DefaultNotificationText_NoInternet=("es-419","No Internet Available")
++DefaultNotificationText_NoInternet=("fr","No Internet Available")
++DefaultNotificationText_NoInternet=("it","No Internet Available")
++DefaultNotificationText_NoInternet=("ja","No Internet Available")
++DefaultNotificationText_NoInternet=("ko","No Internet Available")
++DefaultNotificationText_NoInternet=("pl","No Internet Available")
++DefaultNotificationText_NoInternet=("pt-BR","No Internet Available")
++DefaultNotificationText_NoInternet=("ru","No Internet Available")
++DefaultNotificationText_NoInternet=("tr","No Internet Available")
++DefaultNotificationText_NoInternet=("zh-CN","No Internet Available")
++DefaultNotificationText_NoInternet=("zh-Hant","No Internet Available")
\ No newline at end of file
diff --git a/Engine/Plugins/Online/Android/AndroidFetchBackgroundDownload/Source/AndroidFetchBackgroundDownload/Java/download/UEDownloadWorker.java b/Engine/Plugins/Online/Android/AndroidFetchBackgroundDownload/Source/AndroidFetchBackgroundDownload/Java/download/UEDownloadWorker.java
--- a/Engine/Plugins/Online/Android/AndroidFetchBackgroundDownload/Source/AndroidFetchBackgroundDownload/Java/download/UEDownloadWorker.java
+++ b/Engine/Plugins/Online/Android/AndroidFetchBackgroundDownload/Source/AndroidFetchBackgroundDownload/Java/download/UEDownloadWorker.java
@@ -33,6 +33,9 @@ import com.epicgames.unreal.download.datastructs.DownloadQueueDescription;
 import com.epicgames.unreal.download.datastructs.DownloadWorkerParameterKeys;
 import com.epicgames.unreal.download.fetch.FetchManager;
 
+import com.epicgames.unreal.network.NetworkConnectivityClient;
+import com.epicgames.unreal.network.NetworkChangedManager;
+
 import com.tonyodev.fetch2.Download;
 import com.tonyodev.fetch2.Error;
 import com.tonyodev.fetch2.Request;
@@ -84,6 +87,24 @@ public class UEDownloadWorker extends UEWorker implements DownloadProgressListen
 			NotificationDescription = new DownloadNotificationDescription(getInputData(), getApplicationContext(), Log);
 		}
 
+		if (null == NetworkListener)
+		{
+
+			NetworkListener = new NetworkConnectivityClient.Listener() {
+				@Override
+				public void onNetworkAvailable(NetworkConnectivityClient.NetworkTransportType networkTransportType) {
+					bLostNetwork = false;
+				}
+
+				@Override
+				public void onNetworkLost() {
+					bLostNetwork = true;
+				}
+			};
+			NetworkChangedManager.getInstance().addListener(NetworkListener);
+		}
+
+
 		bHasEnqueueHappened = false;
 		bForceStopped = false;
 	}
@@ -204,15 +225,15 @@ public class UEDownloadWorker extends UEWorker implements DownloadProgressListen
 		{
 			NotificationDescription.CurrentProgress = CurrentProgress;
 			NotificationDescription.Indeterminate = Indeterminate;
-
 			setForegroundAsync(CreateForegroundInfo(NotificationDescription));
+			
 		}
 		else
 		{
 			Log.error("Unexpected NULL NotificationDescripton during UpdateNotification!");
 		}
 	}
-	
+
 	@NonNull
 	private ForegroundInfo CreateForegroundInfo(DownloadNotificationDescription Description) 
 	{		
@@ -221,6 +242,37 @@ public class UEDownloadWorker extends UEWorker implements DownloadProgressListen
 		
 		CreateNotificationChannel(context, notificationManager, Description);
 
+		//Setup Opening UE app if clicked
+		PendingIntent pendingNotificationIntent = null;
+		{
+			Intent notificationIntent = new Intent(context, GameActivity.class);
+			
+			// launch if closed but resume if running
+			notificationIntent.setFlags(Intent.FLAG_ACTIVITY_CLEAR_TOP | Intent.FLAG_ACTIVITY_SINGLE_TOP);
+
+			notificationIntent.putExtra("localNotificationID" , Description.NotificationID);
+			notificationIntent.putExtra("localNotificationAppLaunched" , true);
+
+			pendingNotificationIntent = PendingIntent.getActivity(context, Description.NotificationID, notificationIntent, PendingIntent.FLAG_IMMUTABLE);
+		}
+		
+		Notification notification = null;
+		if(!bLostNetwork)
+		{
+			notification = CreateDownloadProgressNotification(context, Description, pendingNotificationIntent);
+		}
+		else
+		{
+			notification = CreateNoInternetDownloadNotification(context, Description, pendingNotificationIntent);
+		}
+		
+		return new ForegroundInfo(Description.NotificationID,notification);
+	}
+
+	public Notification CreateNoInternetDownloadNotification(Context context,
+		DownloadNotificationDescription Description,
+		PendingIntent pendingNotificationIntent)
+	{
 		//Setup Notification Text Values (ContentText and ContentInfo)
 		boolean bIsComplete = (Description.CurrentProgress >= Description.MAX_PROGRESS);
 		String NotificationTextToUse = null;
@@ -232,19 +284,36 @@ public class UEDownloadWorker extends UEWorker implements DownloadProgressListen
 		{
 			NotificationTextToUse = Description.ContentCompleteText;
 		}
-		
-		//Setup Opening UE app if clicked
-		PendingIntent pendingNotificationIntent = null;
-		{
-			Intent notificationIntent = new Intent(context, GameActivity.class);
-			
-			// launch if closed but resume if running
-			notificationIntent.setFlags(Intent.FLAG_ACTIVITY_CLEAR_TOP | Intent.FLAG_ACTIVITY_SINGLE_TOP);
 
-			notificationIntent.putExtra("localNotificationID" , Description.NotificationID);
-			notificationIntent.putExtra("localNotificationAppLaunched" , true);
+		Notification notification = new NotificationCompat.Builder(context, Description.NotificationChannelID)
+			.setContentTitle(Description.NoInternetAvailable)
+			.setTicker(Description.TitleText)
+			.setContentText(NotificationTextToUse)
+			.setContentIntent(pendingNotificationIntent)
+			.setProgress(Description.MAX_PROGRESS, Description.CurrentProgress, Description.Indeterminate)
+			.setOngoing(true)
+			.setOnlyAlertOnce (true)
+			.setSmallIcon(Description.SmallIconResourceID)
+			.addAction(Description.CancelIconResourceID, Description.CancelText, CancelIntent)
+			.build();
 
-			pendingNotificationIntent = PendingIntent.getActivity(context, Description.NotificationID, notificationIntent, PendingIntent.FLAG_IMMUTABLE);
+		return notification;
+	}
+
+	public Notification CreateDownloadProgressNotification(Context context,
+		DownloadNotificationDescription Description,
+		PendingIntent pendingNotificationIntent)
+	{
+		//Setup Notification Text Values (ContentText and ContentInfo)
+		boolean bIsComplete = (Description.CurrentProgress >= Description.MAX_PROGRESS);
+		String NotificationTextToUse = null;
+		if (!bIsComplete)
+		{
+			NotificationTextToUse = String.format(Description.ContentText, Description.CurrentProgress);
+		}
+		else
+		{
+			NotificationTextToUse = Description.ContentCompleteText;
 		}
 
 		Notification notification = new NotificationCompat.Builder(context, Description.NotificationChannelID)
@@ -252,14 +321,14 @@ public class UEDownloadWorker extends UEWorker implements DownloadProgressListen
 			.setTicker(Description.TitleText)
 			.setContentText(NotificationTextToUse)
 			.setContentIntent(pendingNotificationIntent)
-			.setProgress(Description.MAX_PROGRESS,Description.CurrentProgress, Description.Indeterminate)
+			.setProgress(Description.MAX_PROGRESS, Description.CurrentProgress, Description.Indeterminate)
 			.setOngoing(true)
 			.setOnlyAlertOnce (true)
 			.setSmallIcon(Description.SmallIconResourceID)
 			.addAction(Description.CancelIconResourceID, Description.CancelText, CancelIntent)
 			.build();
-		
-		return new ForegroundInfo(Description.NotificationID,notification);
+
+		return notification;
 	}
 
 	//Gets the Notification Manager through the appropriate method based on build version
@@ -405,6 +474,8 @@ public class UEDownloadWorker extends UEWorker implements DownloadProgressListen
 	public native void nativeAndroidBackgroundDownloadOnAllComplete(boolean bDidAllRequestsSucceed);
 	public native void nativeAndroidBackgroundDownloadOnTick();
 	
+	private NetworkConnectivityClient.Listener NetworkListener = null;
+	private boolean bLostNetwork = false;
 	private boolean bForceStopped = false;
 	private DownloadQueueDescription QueueDescription = null;
 	private volatile boolean bHasEnqueueHappened = false;
diff --git a/Engine/Plugins/Online/Android/AndroidFetchBackgroundDownload/Source/AndroidFetchBackgroundDownload/Java/download/datastructs/DownloadNotificationDescription.java b/Engine/Plugins/Online/Android/AndroidFetchBackgroundDownload/Source/AndroidFetchBackgroundDownload/Java/download/datastructs/DownloadNotificationDescription.java
--- a/Engine/Plugins/Online/Android/AndroidFetchBackgroundDownload/Source/AndroidFetchBackgroundDownload/Java/download/datastructs/DownloadNotificationDescription.java
+++ b/Engine/Plugins/Online/Android/AndroidFetchBackgroundDownload/Source/AndroidFetchBackgroundDownload/Java/download/datastructs/DownloadNotificationDescription.java
@@ -83,6 +83,12 @@ public class DownloadNotificationDescription
 			{
 				CancelText = "Cancel";
 			}
+
+			NoInternetAvailable = data.getString(DownloadWorkerParameterKeys.NOTIFICATION_CONTENT_NO_INTERNET_TEXT_KEY);
+			if (null == NoInternetAvailable)
+			{
+				NoInternetAvailable = "No Internet Available";
+			}
 		}
 		
 		//Load the Cancel Icon Resource
@@ -226,6 +232,7 @@ public class DownloadNotificationDescription
 	public String ContentText = null;
 	public String ContentCompleteText = null;
 	public String CancelText = null;
+	public String NoInternetAvailable = null;
 	
 	public int CancelIconResourceID = 0;
 	public int SmallIconResourceID = 0;
diff --git a/Engine/Plugins/Online/Android/AndroidFetchBackgroundDownload/Source/AndroidFetchBackgroundDownload/Java/download/datastructs/DownloadWorkerParameterKeys.java b/Engine/Plugins/Online/Android/AndroidFetchBackgroundDownload/Source/AndroidFetchBackgroundDownload/Java/download/datastructs/DownloadWorkerParameterKeys.java
--- a/Engine/Plugins/Online/Android/AndroidFetchBackgroundDownload/Source/AndroidFetchBackgroundDownload/Java/download/datastructs/DownloadWorkerParameterKeys.java
+++ b/Engine/Plugins/Online/Android/AndroidFetchBackgroundDownload/Source/AndroidFetchBackgroundDownload/Java/download/datastructs/DownloadWorkerParameterKeys.java
@@ -33,6 +33,7 @@ public class DownloadWorkerParameterKeys
 	public final static String NOTIFICATION_CONTENT_TITLE_KEY = "NotificationContentTitle";
 	public final static String NOTIFICATION_CONTENT_TEXT_KEY = "NotificationContentText";
 	public final static String NOTIFICATION_CONTENT_CANCEL_DOWNLOAD_TEXT_KEY = "NotificationContentCancelDownloadText";
+	public final static String NOTIFICATION_CONTENT_NO_INTERNET_TEXT_KEY = "NotificationContentNoInternetDownloadText";
 	public final static String NOTIFICATION_CONTENT_COMPLETE_TEXT_KEY = "NotificationContentCompleteText";
 
 	public final static String NOTIFICATION_RESOURCE_CANCEL_ICON_NAME = "NotificationResourceCancelIconName";
diff --git a/Engine/Plugins/Online/Android/AndroidFetchBackgroundDownload/Source/AndroidFetchBackgroundDownload/Private/AndroidPlatformBackgroundHttpManager.cpp b/Engine/Plugins/Online/Android/AndroidFetchBackgroundDownload/Source/AndroidFetchBackgroundDownload/Private/AndroidPlatformBackgroundHttpManager.cpp
--- a/Engine/Plugins/Online/Android/AndroidFetchBackgroundDownload/Source/AndroidFetchBackgroundDownload/Private/AndroidPlatformBackgroundHttpManager.cpp
+++ b/Engine/Plugins/Online/Android/AndroidFetchBackgroundDownload/Source/AndroidFetchBackgroundDownload/Private/AndroidPlatformBackgroundHttpManager.cpp
@@ -66,6 +66,7 @@ const int FAndroidNativeDownloadWorkerParameterKeys::NOTIFICATION_DEFAULT_ID_KEY
 const FString FAndroidNativeDownloadWorkerParameterKeys::NOTIFICATION_CONTENT_TITLE_KEY = TEXT("NotificationContentTitle");
 const FString FAndroidNativeDownloadWorkerParameterKeys::NOTIFICATION_CONTENT_TEXT_KEY = TEXT("NotificationContentText");
 const FString FAndroidNativeDownloadWorkerParameterKeys::NOTIFICATION_CONTENT_CANCEL_DOWNLOAD_TEXT_KEY = TEXT("NotificationContentCancelDownloadText");
+const FString FAndroidNativeDownloadWorkerParameterKeys::NOTIFICATION_CONTENT_NO_INTERNET_TEXT_KEY = TEXT("NotificationContentNoInternetDownloadText");
 const FString FAndroidNativeDownloadWorkerParameterKeys::NOTIFICATION_CONTENT_COMPLETE_TEXT_KEY = TEXT("NotificationContentCompleteText");
 
 const FString FAndroidNativeDownloadWorkerParameterKeys::NOTIFICATION_RESOURCE_CANCEL_ICON_NAME = TEXT("NotificationResourceCancelIconName");
@@ -306,7 +307,8 @@ void FAndroidPlatformBackgroundHttpManager::ActivatePendingRequests()
 			{
 				FUEWorkManagerNativeWrapper::FWorkRequestParametersNative WorkParams;
 				WorkParams.WorkerJavaClass = FAndroidPlatformBackgroundHttpManager::JavaInfo.UEDownloadWorkerClass;
-				WorkParams.bRequireAnyInternet = true;
+				//Setting this to false to avoid notifications clearing when dropping Internet connection.
+				WorkParams.bRequireAnyInternet = false;
 				WorkParams.bStartAsForegroundService = true;
 					
 				//Set our DownloadDescription file so that it can be parsed by the worker
@@ -319,6 +321,7 @@ void FAndroidPlatformBackgroundHttpManager::ActivatePendingRequests()
 				WorkParams.AddDataToWorkerParameters(FAndroidNativeDownloadWorkerParameterKeys::NOTIFICATION_CONTENT_TITLE_KEY, AndroidBackgroundHTTPManagerDefaultLocalizedText.DefaultNotificationText_Title.GetText());
 				WorkParams.AddDataToWorkerParameters(FAndroidNativeDownloadWorkerParameterKeys::NOTIFICATION_CONTENT_COMPLETE_TEXT_KEY, AndroidBackgroundHTTPManagerDefaultLocalizedText.DefaultNotificationText_Complete.GetText());
 				WorkParams.AddDataToWorkerParameters(FAndroidNativeDownloadWorkerParameterKeys::NOTIFICATION_CONTENT_CANCEL_DOWNLOAD_TEXT_KEY, AndroidBackgroundHTTPManagerDefaultLocalizedText.DefaultNotificationText_Cancel.GetText());
+				WorkParams.AddDataToWorkerParameters(FAndroidNativeDownloadWorkerParameterKeys::NOTIFICATION_CONTENT_NO_INTERNET_TEXT_KEY, AndroidBackgroundHTTPManagerDefaultLocalizedText.DefaultNotificationText_NoInternet.GetText());
 
 				//Expect our ContentText to have a {DownloadPercent} argument in it by default, so this will replace that with the Java string format argument so Java can insert the appropriate value
 				FFormatNamedArguments Arguments;
@@ -941,6 +944,7 @@ FAndroidPlatformBackgroundHttpManager::FAndroidBackgroundHTTPManagerDefaultLocal
 	, DefaultNotificationText_Content()
 	, DefaultNotificationText_Complete()
 	, DefaultNotificationText_Cancel()
+	, DefaultNotificationText_NoInternet()
 {
 }
 
@@ -980,6 +984,14 @@ void FAndroidPlatformBackgroundHttpManager::FAndroidBackgroundHTTPManagerDefault
 				Config->GetArray(*TextConfigSection, TEXT("DefaultNotificationText_Cancel"), CancelTextStrings);
 				ParsePolyglotTextItem(DefaultNotificationText_Cancel, TEXT("Notification.CancelText"), CancelTextStrings);
 			}
+
+			//DefaultNotificationText_NoInternet
+			{
+				TArray<FString> NoInternetTextStrings;
+				Config->GetArray(*TextConfigSection, TEXT("DefaultNotificationText_NoInternet"), NoInternetTextStrings);
+				ParsePolyglotTextItem(DefaultNotificationText_NoInternet, TEXT("Notification.NoInternetText"), NoInternetTextStrings);
+			}
+
 		}
 	}
 
@@ -988,6 +1000,7 @@ void FAndroidPlatformBackgroundHttpManager::FAndroidBackgroundHTTPManagerDefault
 	ForceValidPolyglotText(DefaultNotificationText_Content, TEXT("DefaultNotificationText_Content"));
 	ForceValidPolyglotText(DefaultNotificationText_Complete, TEXT("DefaultNotificationText_Complete"));
 	ForceValidPolyglotText(DefaultNotificationText_Cancel, TEXT("DefaultNotificationText_Cancel"));
+	ForceValidPolyglotText(DefaultNotificationText_NoInternet, TEXT("DefaultNotificationText_NoInternet"));
 }
 
 void FAndroidPlatformBackgroundHttpManager::FAndroidBackgroundHTTPManagerDefaultLocalizedText::ForceValidPolyglotText(FPolyglotTextData& TextDataOut, const FString& DebugPolyglotTextName)
diff --git a/Engine/Plugins/Online/Android/AndroidFetchBackgroundDownload/Source/AndroidFetchBackgroundDownload/Public/AndroidPlatformBackgroundHttpManager.h b/Engine/Plugins/Online/Android/AndroidFetchBackgroundDownload/Source/AndroidFetchBackgroundDownload/Public/AndroidPlatformBackgroundHttpManager.h
--- a/Engine/Plugins/Online/Android/AndroidFetchBackgroundDownload/Source/AndroidFetchBackgroundDownload/Public/AndroidPlatformBackgroundHttpManager.h
+++ b/Engine/Plugins/Online/Android/AndroidFetchBackgroundDownload/Source/AndroidFetchBackgroundDownload/Public/AndroidPlatformBackgroundHttpManager.h
@@ -149,6 +149,8 @@ protected:
 		FPolyglotTextData DefaultNotificationText_Content;
 		FPolyglotTextData DefaultNotificationText_Complete;
 		FPolyglotTextData DefaultNotificationText_Cancel;
+		FPolyglotTextData DefaultNotificationText_NoInternet;
+
 
 		void InitFromIniSettings(const FString& ConfigFileName);
 
@@ -212,6 +214,7 @@ public:
 	static const FString NOTIFICATION_CONTENT_TITLE_KEY;
 	static const FString NOTIFICATION_CONTENT_TEXT_KEY;
 	static const FString NOTIFICATION_CONTENT_CANCEL_DOWNLOAD_TEXT_KEY;
+	static const FString NOTIFICATION_CONTENT_NO_INTERNET_TEXT_KEY;
 	static const FString NOTIFICATION_CONTENT_COMPLETE_TEXT_KEY;
 
 	static const FString NOTIFICATION_RESOURCE_CANCEL_ICON_NAME;
