Improve Android UEDownloadWorker and FetchManager thread safety.
#tests Android Dev Client

diff --git a/Engine/Plugins/Online/Android/AndroidFetchBackgroundDownload/Source/AndroidFetchBackgroundDownload/Java/download/UEDownloadWorker.java b/Engine/Plugins/Online/Android/AndroidFetchBackgroundDownload/Source/AndroidFetchBackgroundDownload/Java/download/UEDownloadWorker.java
--- a/Engine/Plugins/Online/Android/AndroidFetchBackgroundDownload/Source/AndroidFetchBackgroundDownload/Java/download/UEDownloadWorker.java
+++ b/Engine/Plugins/Online/Android/AndroidFetchBackgroundDownload/Source/AndroidFetchBackgroundDownload/Java/download/UEDownloadWorker.java
@@ -479,7 +479,7 @@ public class UEDownloadWorker extends UEWorker implements DownloadProgressListen
 	private boolean bForceStopped = false;
 	private DownloadQueueDescription QueueDescription = null;
 	private volatile boolean bHasEnqueueHappened = false;
-	static FetchManager mFetchManager;
+	static volatile FetchManager mFetchManager;
 	private PendingIntent CancelIntent = null;
 	private DownloadNotificationDescription NotificationDescription = null;
 }
\ No newline at end of file
diff --git a/Engine/Plugins/Online/Android/AndroidFetchBackgroundDownload/Source/AndroidFetchBackgroundDownload/Java/download/fetch/FetchManager.java b/Engine/Plugins/Online/Android/AndroidFetchBackgroundDownload/Source/AndroidFetchBackgroundDownload/Java/download/fetch/FetchManager.java
--- a/Engine/Plugins/Online/Android/AndroidFetchBackgroundDownload/Source/AndroidFetchBackgroundDownload/Java/download/fetch/FetchManager.java
+++ b/Engine/Plugins/Online/Android/AndroidFetchBackgroundDownload/Source/AndroidFetchBackgroundDownload/Java/download/fetch/FetchManager.java
@@ -53,14 +53,18 @@ public class FetchManager implements FetchDownloadProgressOwner, FetchEnqueueRes
 	{
 		Log.debug("StopWork called");
 
-		if (IsFetchInstanceValid())
+		//Synchronized with initFetch to make sure we aren't creating and closing our fetch instance in a weird race condition
+		synchronized(this)
 		{
-			//We need to purge all partially completed downloads
-			//If we don't, then on relaunching our app the BackgroundHttp system can end up assuming these downloads are finished incorrectly.
-			DeleteAllInProgressRequests();
+			if (IsFetchInstanceValid())
+			{
+				//We need to purge all partially completed downloads
+				//If we don't, then on relaunching our app the BackgroundHttp system can end up assuming these downloads are finished incorrectly.
+				DeleteAllInProgressRequests();
 
-			//Freeze our FetchInstance so it stops all work until we unfreeze on resuming work
-			FetchInstance.close();
+				//close our FetchInstance so it stops all work until we recreate it in initfetch later
+				FetchInstance.close();
+			}
 		}
 	}
 	
@@ -166,7 +170,7 @@ public class FetchManager implements FetchDownloadProgressOwner, FetchEnqueueRes
 			
 				FetchInstance.enqueue(FetchRequest, RequestCallback, ErrorCallback);
 
-				Log.debug("Enqueued request:" + Description.RequestID);
+				Log.debug("Enqueue request sent for:" + Description.RequestID);
 			}
 		}
 		//if we have previously completed this download, then we just want to compelte it instead of creating a new FetchRequest
@@ -335,56 +339,60 @@ public class FetchManager implements FetchDownloadProgressOwner, FetchEnqueueRes
 	{
 		Log.debug("InitFetch called");
 
-		//Make sure any existing FetchInstance is in a correct state (either null and ready to be created, or open and unfrozen ready to do work
-		if (FetchInstance != null)
+		//Synchronized with StopWork to make sure we aren't creating and closing our fetch instance in a weird race condition
+		synchronized(this)
 		{
-			//If we previously closed our FetchInstance just remove it and recreate it bellow
-			if (FetchInstance.isClosed())
-			{
-				FetchInstance = null;
-			}
-			//If our FetchInstance exists and isn't closed, its very likely frozen and needs to be unfrozen
-			else
+			//Make sure any existing FetchInstance is in a correct state (either null and ready to be created, or open and unfrozen ready to do work
+			if (FetchInstance != null)
 			{
-				Log.debug("InitFetch has existing non-closed FetchInstance. Unfreezing and deleting all in progress requests");
+				//If we previously closed our FetchInstance just remove it and recreate it bellow
+				if (FetchInstance.isClosed())
+				{
+					FetchInstance = null;
+				}
+				//If our FetchInstance exists and isn't closed, its very likely frozen and needs to be unfrozen
+				else
+				{
+					Log.debug("InitFetch has existing non-closed FetchInstance. Unfreezing and deleting all in progress requests");
 
-				//If we are just unfreezing existing Fetch work, lets delete everything that wasn't finished by the previous work before unfreezing
-				//This prevents errors in resuming where we no longer want a particular download or we fail to resume the work
-				DeleteAllInProgressRequests();
+					//If we are just unfreezing existing Fetch work, lets delete everything that wasn't finished by the previous work before unfreezing
+					//This prevents errors in resuming where we no longer want a particular download or we fail to resume the work
+					DeleteAllInProgressRequests();
 
-				//Now unfreeze (and hopefully have nothing really running)
-				FetchInstance.unfreeze();
+					//Now unfreeze (and hopefully have nothing really running)
+					FetchInstance.unfreeze();
+				}
 			}
-		}
-
-		if (FetchInstance == null)
-		{
-//			InstanceLogger = new FetchLogger(true, "FManager");
 
-			//TODO TRoss: Pull these values from the worker's getInputData
-			FetchInstance = Fetch.Impl.getInstance(new FetchConfiguration.Builder(context)
-				.setNamespace(context.getPackageName())
-//				.enableLogging(true)
-//				.setLogger(InstanceLogger)
-				.enableRetryOnNetworkGain(true)
-				.setProgressReportingInterval(200)
-				.build());
-
-			//if we are creating our FetchInstance, make sure our FetchListener is also recreated and attached
-			FetchListener = null;
-		}
+			if (FetchInstance == null)
+			{
+				//InstanceLogger = new FetchLogger(true, "FManager");
+
+				//TODO TRoss: Pull these values from the worker's getInputData
+				FetchInstance = Fetch.Impl.getInstance(new FetchConfiguration.Builder(context)
+					.setNamespace(context.getPackageName())
+					//.enableLogging(true)
+					//.setLogger(InstanceLogger)
+					.enableRetryOnNetworkGain(true)
+					.setProgressReportingInterval(200)
+					.build());
+
+				//if we are creating our FetchInstance, make sure our FetchListener is also recreated and attached
+				FetchListener = null;
+			}
 					
-		if (!IsFetchInstanceValid())
-		{
-			Log.error("Unexpected invalid FetchInstance after completing InitFetch!");
-		}
-		else
-		{
-			//Add our FetchListener
-			if (null == FetchListener)
+			if (!IsFetchInstanceValid())
 			{
-				FetchListener = new FetchRequestProgressListener(this);
-				FetchInstance.addListener(FetchListener);
+				Log.error("Unexpected invalid FetchInstance after completing InitFetch!");
+			}
+			else
+			{
+				//Add our FetchListener
+				if (null == FetchListener)
+				{
+					FetchListener = new FetchRequestProgressListener(this);
+					FetchInstance.addListener(FetchListener);
+				}
 			}
 		}
 	}
@@ -956,14 +964,14 @@ public class FetchManager implements FetchDownloadProgressOwner, FetchEnqueueRes
 	}
 	//private boolean ShouldRetryRequest(@NonNull Request 
 	
-	private Fetch FetchInstance = null;
+	private volatile Fetch FetchInstance = null;
 	private FetchLogger InstanceLogger = null;
 
 	private FetchRequestProgressListener FetchListener = null;
 	
-	private HashMap<String, DownloadDescription> RequestedDownloads = new HashMap<String, DownloadDescription>();
-	private HashMap<String, DownloadDescription> CompletedDownloads = new HashMap<String, DownloadDescription>();
-	private HashMap<String, DownloadDescription> FailedDownloads = new HashMap<String, DownloadDescription>();
+	private volatile HashMap<String, DownloadDescription> RequestedDownloads = new HashMap<String, DownloadDescription>();
+	private volatile HashMap<String, DownloadDescription> CompletedDownloads = new HashMap<String, DownloadDescription>();
+	private volatile HashMap<String, DownloadDescription> FailedDownloads = new HashMap<String, DownloadDescription>();
 
 	public String TempFileExtension = ".fetchtemp";
 	
