Update Background Downloading Notification localization for No internet available

diff --git a/Engine/Plugins/Online/Android/AndroidFetchBackgroundDownload/Config/BaseAndroidFetchBackgroundDownload.ini b/Engine/Plugins/Online/Android/AndroidFetchBackgroundDownload/Config/BaseAndroidFetchBackgroundDownload.ini
--- a/Engine/Plugins/Online/Android/AndroidFetchBackgroundDownload/Config/BaseAndroidFetchBackgroundDownload.ini
+++ b/Engine/Plugins/Online/Android/AndroidFetchBackgroundDownload/Config/BaseAndroidFetchBackgroundDownload.ini
@@ -71,18 +71,16 @@
 +DefaultNotificationText_Cancel=("zh-Hant","取消")
 
 ;NoInternet Text
-+DefaultNotificationText_NoInternet=("native","No Internet Available")
-+DefaultNotificationText_NoInternet=("ar","No Internet Available")
-+DefaultNotificationText_NoInternet=("de","Kein Internetzugang")
-+DefaultNotificationText_NoInternet=("es","No Internet Available")
-+DefaultNotificationText_NoInternet=("es-419","No Internet Available")
-+DefaultNotificationText_NoInternet=("fr","No Internet Available")
-+DefaultNotificationText_NoInternet=("it","No Internet Available")
-+DefaultNotificationText_NoInternet=("ja","No Internet Available")
-+DefaultNotificationText_NoInternet=("ko","No Internet Available")
-+DefaultNotificationText_NoInternet=("pl","No Internet Available")
-+DefaultNotificationText_NoInternet=("pt-BR","No Internet Available")
-+DefaultNotificationText_NoInternet=("ru","No Internet Available")
-+DefaultNotificationText_NoInternet=("tr","No Internet Available")
-+DefaultNotificationText_NoInternet=("zh-CN","No Internet Available")
-+DefaultNotificationText_NoInternet=("zh-Hant","No Internet Available")
\ No newline at end of file
++DefaultNotificationText_NoInternet=("native","No Internet available")
++DefaultNotificationText_NoInternet=("ar","الإنترنت غير متوفر")
++DefaultNotificationText_NoInternet=("de","Keine Internetverbindung verfügbar")
++DefaultNotificationText_NoInternet=("es","No hay conexión a internet")
++DefaultNotificationText_NoInternet=("es-419","No hay conexión a internet")
++DefaultNotificationText_NoInternet=("fr","Internet indisponible")
++DefaultNotificationText_NoInternet=("it","Internet non disponibile")
++DefaultNotificationText_NoInternet=("ja","インターネットが利用できません")
++DefaultNotificationText_NoInternet=("ko","인터넷이 없습니다")
++DefaultNotificationText_NoInternet=("pl","Brak internetu")
++DefaultNotificationText_NoInternet=("pt-BR","Internet indisponível")
++DefaultNotificationText_NoInternet=("ru","Нет доступа в интернет")
++DefaultNotificationText_NoInternet=("tr","İnternet bağlantısı yok")
\ No newline at end of file
