Make Background downloading % reflect the actual bytes downloaded

diff --git a/Engine/Plugins/Online/Android/AndroidFetchBackgroundDownload/Source/AndroidFetchBackgroundDownload/Java/download/datastructs/DownloadDescription.java b/Engine/Plugins/Online/Android/AndroidFetchBackgroundDownload/Source/AndroidFetchBackgroundDownload/Java/download/datastructs/DownloadDescription.java
--- a/Engine/Plugins/Online/Android/AndroidFetchBackgroundDownload/Source/AndroidFetchBackgroundDownload/Java/download/datastructs/DownloadDescription.java
+++ b/Engine/Plugins/Online/Android/AndroidFetchBackgroundDownload/Source/AndroidFetchBackgroundDownload/Java/download/datastructs/DownloadDescription.java
@@ -44,6 +44,8 @@ public class DownloadDescription
 	public boolean bIsCancelled = false;
 	public long PreviousDownloadedBytes = 0;
 	public int PreviousDownloadPercent = 0;
+	public long TotalBytesNeeded = 0;
+	public long TotalDownloadedBytes = 0;
 	public DownloadProgressListener ProgressListener = null;
 	public DownloadDescription()
 	{
diff --git a/Engine/Plugins/Online/Android/AndroidFetchBackgroundDownload/Source/AndroidFetchBackgroundDownload/Java/download/fetch/FetchManager.java b/Engine/Plugins/Online/Android/AndroidFetchBackgroundDownload/Source/AndroidFetchBackgroundDownload/Java/download/fetch/FetchManager.java
--- a/Engine/Plugins/Online/Android/AndroidFetchBackgroundDownload/Source/AndroidFetchBackgroundDownload/Java/download/fetch/FetchManager.java
+++ b/Engine/Plugins/Online/Android/AndroidFetchBackgroundDownload/Source/AndroidFetchBackgroundDownload/Java/download/fetch/FetchManager.java
@@ -167,6 +167,12 @@ public class FetchManager implements FetchDownloadProgressOwner, FetchEnqueueRes
 			{
 				Request FetchRequest = BuildFetchRequest(Description);
 				Description.CachedFetchID = FetchRequest.getId();
+
+				//Needed for callbacks
+				DownloadContentLength downContent = new DownloadContentLength();
+				downContent.desc = Description;
+				DownloadContentLengthError downContentError = new DownloadContentLengthError();
+				FetchInstance.getContentLengthForRequest(FetchRequest, true, downContent, downContentError);
 			
 				FetchInstance.enqueue(FetchRequest, RequestCallback, ErrorCallback);
 
@@ -433,6 +439,24 @@ public class FetchManager implements FetchDownloadProgressOwner, FetchEnqueueRes
 		private	DownloadDescription RecreateDescription;
 	}
 
+	private class DownloadContentLength implements Func<Long>
+	{
+		@Override
+		public void call(@NonNull Long result) {
+			desc.TotalBytesNeeded = result;
+		}
+
+		public DownloadDescription desc;
+	}
+
+	// Needed for getContentLengthForRequest.
+	private class DownloadContentLengthError implements Func<Error>
+	{
+		@Override
+		public void call(@NonNull Error result) {
+		}
+	}
+
 	//Helper class to avoid use of Delegates as our current compile source target is 7 and thus delegates are not supported.
 	//Cancels downloads found by a given tag.
 	private class CancelDownloadByTagFunc implements Func<List<Download>>
@@ -473,12 +497,8 @@ public class FetchManager implements FetchDownloadProgressOwner, FetchEnqueueRes
 
 	public void RequestGroupProgressUpdate(int GroupID, DownloadProgressListener ListenerToUpdate)
 	{
-		//For now just assume each download is roughly the same size.
-		//TODO TRoss, we should pass in the expected download amount potentially for each download and actually compute these values
-		//so that we can do more accurate % rather then larger files "slowing down" the progress bar progression.
-		
-		int TotalProgress = 0;
-		int TotalDownloadsInGroup = 0;
+		long TotalBytesProgress = 0;
+		long TotalBytesNeeded = 0;
 		
 		ArrayList<String> DownloadKeys = new ArrayList<String>(RequestedDownloads.keySet());
 		for (int DescIndex = 0; DescIndex < DownloadKeys.size(); ++DescIndex)
@@ -487,21 +507,22 @@ public class FetchManager implements FetchDownloadProgressOwner, FetchEnqueueRes
 			
 			if (FoundDesc.GroupID == GroupID)
 			{
-				++TotalDownloadsInGroup;
-				TotalProgress += FoundDesc.PreviousDownloadPercent;
+				TotalBytesProgress += FoundDesc.TotalDownloadedBytes;
+				TotalBytesNeeded += FoundDesc.TotalBytesNeeded;
 			}
 		}
 				
-		//just get the raw average of this to send back
-		int TotalToSend = TotalProgress / TotalDownloadsInGroup;
-		boolean bIsIndeterminate = (TotalProgress == 0);
+		float Progress = (float)TotalBytesProgress / (float)TotalBytesNeeded;
+		//Based on BPS BackgroundDownload wight is 3
+		Progress = Progress * 0.75f;
+		boolean bIsIndeterminate = (TotalBytesProgress == 0);
 		
 		//Make sure we cap at 100% progress
-		if (TotalToSend > 100)
+		if (Progress > 1.0f)
 		{
-			TotalToSend = 100;
+			Progress = 1.0f;
 		}
-		ListenerToUpdate.OnDownloadGroupProgress(GroupID, TotalToSend, bIsIndeterminate);
+		ListenerToUpdate.OnDownloadGroupProgress(GroupID, (int)Math.ceil(Progress * 100.0f), bIsIndeterminate);
 	}
 
 	//
@@ -545,9 +566,12 @@ public class FetchManager implements FetchDownloadProgressOwner, FetchEnqueueRes
 			{
 				TotalDownloadedSinceLastCall = (TotalDownloaded - MatchedDownload.PreviousDownloadedBytes);
 			}
-
 			MatchedDownload.PreviousDownloadPercent = download.getProgress();
+			MatchedDownload.PreviousDownloadedBytes = TotalDownloadedSinceLastCall;
+			MatchedDownload.TotalDownloadedBytes = TotalDownloaded;
 			MatchedDownload.ProgressListener.OnDownloadProgress(GetRequestID(download), TotalDownloadedSinceLastCall, TotalDownloaded);
+			MatchedDownload.TotalBytesNeeded = download.getTotal();
+			
 		}
 		else
 		{
