Android Background HTTP fixes
- Now support force requeuing existing ActiveRequests in a new UEDownloadWorker when work is stopped on the existing worker unexpectidly (such as when the user pressed the cancel button). This fixes some hard locks in background http when the user or OS cancels the download without the app restarting.
- Fix for unchecked completed download data growth in FetchManager
- Fix for freeze scenarios where RequestedDownloads list in the FetchManager could grow beyond its original intent and thus get stuck waiting to do work that wasn't possible due to duplicates in the list.
- Fixed bug where MarkUnderlyingJavaRequestAsCompleted was not respecting the passed in bool value for success or failure.
- Fix for ensures triggering when resuming the app with active background downloads from a previous run. Now cancels all work during initialize so that when the app is running only downloads started from that app run are active.
#tests Android Dev Client

diff --git a/Engine/Plugins/Online/Android/AndroidFetchBackgroundDownload/Source/AndroidFetchBackgroundDownload/Java/download/fetch/FetchManager.java b/Engine/Plugins/Online/Android/AndroidFetchBackgroundDownload/Source/AndroidFetchBackgroundDownload/Java/download/fetch/FetchManager.java
--- a/Engine/Plugins/Online/Android/AndroidFetchBackgroundDownload/Source/AndroidFetchBackgroundDownload/Java/download/fetch/FetchManager.java
+++ b/Engine/Plugins/Online/Android/AndroidFetchBackgroundDownload/Source/AndroidFetchBackgroundDownload/Java/download/fetch/FetchManager.java
@@ -110,6 +110,13 @@ public class FetchManager implements FetchDownloadProgressOwner, FetchEnqueueRes
 	//don't match our ActiveDownloadDescriptions entries and thus need to be reconciled (either added or modified in some way)
 	private void ReconcileDownloadDescriptions(DownloadQueueDescription QueueDescription)
 	{
+		//Create reference copy of RequestedDownloads and clear original lists for
+		//RequestedDownloads, CompletedDownloads, and FailedDownloads to allow them to re-populate with new work state
+		HashMap<String,DownloadDescription> TempRequestedDownloadsCopy =  new HashMap<String,DownloadDescription>(RequestedDownloads);
+		RequestedDownloads.clear();
+		CompletedDownloads.clear();
+		FailedDownloads.clear();
+
 		final int NumNewDownloadDescriptions = QueueDescription.DownloadDescriptions.size();
 		for (int DescriptionIndex = 0; DescriptionIndex < NumNewDownloadDescriptions; ++DescriptionIndex)
 		{
@@ -119,28 +126,49 @@ public class FetchManager implements FetchDownloadProgressOwner, FetchEnqueueRes
 			NewDownloadDescription.ProgressListener = QueueDescription.ProgressListener;
 							
 			String RequestIDKey = NewDownloadDescription.RequestID;
-			//Don't even have an entry for this RequestID, so it's completely new
-			if (!RequestedDownloads.containsKey(RequestIDKey))
+			//Don't even have an entry for this RequestID in the old list, so it's completely new
+			if (!TempRequestedDownloadsCopy.containsKey(RequestIDKey))
 			{
 				QueueNewDownloadDescription(NewDownloadDescription);
 			}
 			//Need to update our DownloadDescription
 			else
 			{
-				//don't requeue completed downloads
-				if (NewDownloadDescription.bHasCompleted)
-				{
-					Log.debug("Skipping requeue of " + NewDownloadDescription.RequestID + " as its flagged for bHasCompleted");
-				}
-				else
-				{
-					DownloadDescription ActiveDescription = RequestedDownloads.get(RequestIDKey);
-					HandleChangedDownloadDescription(ActiveDescription, NewDownloadDescription);
-				}
+				DownloadDescription ActiveDescription = TempRequestedDownloadsCopy.get(RequestIDKey);
+				HandleChangedDownloadDescription(ActiveDescription, NewDownloadDescription);
 			}
 		}
 	}
 	
+	public void ResetDownloadCompleteTracking(DownloadDescription Description)
+	{
+		if (CompletedDownloads.containsKey(Description.RequestID))
+		{
+			CompletedDownloads.remove(Description.RequestID);
+		}
+
+		if (FailedDownloads.containsKey(Description.RequestID))
+		{
+			FailedDownloads.remove(Description.RequestID);
+		}
+	}
+
+	public ECompleteReason GetCompleteReasonForDownload(DownloadDescription Description)
+	{
+		if (Description.bIsCancelled)
+		{
+			return ECompleteReason.Cancelled;
+		}
+
+		if (FailedDownloads.containsKey(Description.RequestID))
+		{
+			return ECompleteReason.Error;
+		}
+
+		//Default to success since the download is complete and we have no signs of an error
+		return ECompleteReason.Success;
+	}
+
 	public void QueueNewDownloadDescription(DownloadDescription Description)
 	{
 		//We have hit this code after something else has already invalidated our FetchInstance
@@ -151,47 +179,43 @@ public class FetchManager implements FetchDownloadProgressOwner, FetchEnqueueRes
 		}
 		
 		RequestedDownloads.put(Description.RequestID, Description);
+		ResetDownloadCompleteTracking(Description);
 
 		if (!CheckForPreviouslyCompletedDownload(Description))
 		{	
-			//Make sure this request wasn't previously flagged as completed before beginning new work
-			if (CompletedDownloads.containsKey(Description.RequestID))
-			{
-				CompletedDownloads.remove(Description.RequestID);
-			}
-
 			FetchEnqueueResultListener.FetchEnqueueRequestCallback RequestCallback = new FetchEnqueueResultListener.FetchEnqueueRequestCallback(this, Description.RequestID);
 			FetchEnqueueResultListener.FetchEnqueueErrorCallback ErrorCallback = new FetchEnqueueResultListener.FetchEnqueueErrorCallback(this, Description.RequestID);
 		
-			if (!Description.bIsCancelled)
-			{
-				Request FetchRequest = BuildFetchRequest(Description);
-				Description.CachedFetchID = FetchRequest.getId();
-
-				//Needed for callbacks
-				DownloadContentLength downContent = new DownloadContentLength();
-				downContent.desc = Description;
-				DownloadContentLengthError downContentError = new DownloadContentLengthError();
-				FetchInstance.getContentLengthForRequest(FetchRequest, true, downContent, downContentError);
+			Request FetchRequest = BuildFetchRequest(Description);
+			Description.CachedFetchID = FetchRequest.getId();
+
+			//Needed for callbacks
+			DownloadContentLength downContent = new DownloadContentLength();
+			downContent.desc = Description;
+			DownloadContentLengthError downContentError = new DownloadContentLengthError();
+			FetchInstance.getContentLengthForRequest(FetchRequest, true, downContent, downContentError);
 			
-				FetchInstance.enqueue(FetchRequest, RequestCallback, ErrorCallback);
+			FetchInstance.enqueue(FetchRequest, RequestCallback, ErrorCallback);
 
-				Log.debug("Enqueue request sent for:" + Description.RequestID);
-			}
+			Log.debug("Enqueue request sent for:" + Description.RequestID);
 		}
 		//if we have previously completed this download, then we just want to compelte it instead of creating a new FetchRequest
 		else
 		{
-			Log.debug("Completed RequestID:" + Description.RequestID + " as the DestinationLocation already is occupied and thus previously completed.");
-			CompleteDownload(Description, ECompleteReason.Success);
+			ECompleteReason CompleteReason = GetCompleteReasonForDownload(Description);
+			Log.debug("Completed RequestID:" + Description.RequestID + " with Reason " + CompleteReason + " as request was previously completed.");
+			CompleteDownload(Description, CompleteReason);
 		}
 	}
 
 	private boolean CheckForPreviouslyCompletedDownload(DownloadDescription Description)
 	{
+		boolean bInCompleteList = CompletedDownloads.containsKey(Description.RequestID);
+		
 		File CompleteLocation = new File(Description.DestinationLocation);
-		return CompleteLocation.exists();
+		boolean bHasFileInCompleteLocation = CompleteLocation.exists();
 
+		return (bInCompleteList || bHasFileInCompleteLocation || Description.bHasCompleted || Description.bIsCancelled);
 	}
 
 	public void PauseDownload(String RequestID, boolean bPause)
@@ -291,12 +315,13 @@ public class FetchManager implements FetchDownloadProgressOwner, FetchEnqueueRes
 	{
 		NewDescription.CurrentRetryCount = OldDescription.CurrentRetryCount;
 		NewDescription.CachedFetchID = OldDescription.CachedFetchID;
-		NewDescription.bIsPaused = OldDescription.bIsPaused;
 		NewDescription.PreviousDownloadedBytes = OldDescription.PreviousDownloadedBytes;
 		NewDescription.PreviousDownloadPercent = OldDescription.PreviousDownloadPercent;
-		
-		//if this has shown back up, we don't want to copy bIsCancelled because we want to re-queue this download
-		NewDescription.bIsCancelled = false;
+		NewDescription.TotalBytesNeeded = OldDescription.TotalBytesNeeded;
+		NewDescription.TotalDownloadedBytes = OldDescription.TotalDownloadedBytes;
+
+		// Purposefully don't copy bIsPaused or bIsCancelled as we want that to come from the new description in case 
+		// this has been set on C++ side but not yet received by our Worker before new work was queued. C++ state should be definitive
 
 		//Purposefully don't copy DownloadProgressListener as we set this based on the new QueueDescription already and likely
 		//don't want to keep the old one (although in current code they match).
@@ -910,12 +935,6 @@ public class FetchManager implements FetchDownloadProgressOwner, FetchEnqueueRes
 			FailedDownloads.put(DownloadDesc.RequestID, DownloadDesc);
 		}
 
-		//Flag any successful complete so we don't redownload it on retries
-		if (CompleteReasonToSend == EDownloadCompleteReason.Success)
-		{
-			DownloadDesc.bHasCompleted = true;
-		}
-
 		DownloadDesc.ProgressListener.OnDownloadComplete(DownloadDesc.RequestID, DownloadDesc.DestinationLocation, CompleteReasonToSend);		
 		CheckForAllDownloadsComplete(DownloadDesc.ProgressListener);
 	}
diff --git a/Engine/Plugins/Online/Android/AndroidFetchBackgroundDownload/Source/AndroidFetchBackgroundDownload/Private/AndroidPlatformBackgroundHttpManager.cpp b/Engine/Plugins/Online/Android/AndroidFetchBackgroundDownload/Source/AndroidFetchBackgroundDownload/Private/AndroidPlatformBackgroundHttpManager.cpp
--- a/Engine/Plugins/Online/Android/AndroidFetchBackgroundDownload/Source/AndroidFetchBackgroundDownload/Private/AndroidPlatformBackgroundHttpManager.cpp
+++ b/Engine/Plugins/Online/Android/AndroidFetchBackgroundDownload/Source/AndroidFetchBackgroundDownload/Private/AndroidPlatformBackgroundHttpManager.cpp
@@ -17,6 +17,8 @@
 #include "Android/AndroidJNI.h"
 #include "Android/AndroidApplication.h"
 
+#include "Containers/UnrealString.h"
+
 #include "Misc/Paths.h"
 
 #if UE_BUILD_SHIPPING
@@ -40,17 +42,31 @@
 
 #define LOCTEXT_NAMESPACE "AndroidBackgroundHttpManager"
 
+//Call backs called by the JNI java systems to call back into C++ code in this AndroidPlatformBackgroundHttpManager
+namespace FAndroidBackgroundDownloadDelegates
+{
+	DECLARE_DELEGATE_TwoParams(FAndroidBackgroundDownload_OnWorkerStart, FString /*WorkID*/, jobject /*UnderlyingWorker*/);
+	DECLARE_DELEGATE_TwoParams(FAndroidBackgroundDownload_OnWorkerStop, FString /*WorkID*/, jobject /*UnderlyingWorker*/);
+	DECLARE_DELEGATE_FourParams(FAndroidBackgroundDownload_OnProgress, jobject /*UnderlyingWorker*/, FString /*RequestID*/, int64_t /*BytesWrittenSinceLastCall*/, int64_t /*TotalBytesWritten*/);
+	DECLARE_DELEGATE_FourParams(FAndroidBackgroundDownload_OnComplete, jobject /*UnderlyingWorker*/, FString /*RequestID*/, FString /*CompleteLocation*/, bool /*bWasSuccess*/);
+	DECLARE_DELEGATE_TwoParams(FAndroidBackgroundDownload_OnAllComplete, jobject /*UnderlyingWorker*/, bool /*bDidAllRequestsSucceed*/);
+	DECLARE_DELEGATE_TwoParams(FAndroidBackgroundDownload_OnTickWorkerThread, JNIEnv*, jobject /*UnderlyingWorker*/);
+	DECLARE_DELEGATE_TwoParams(FAndroidBackgroundDownload_OnNetworkChanged, JNIEnv*, bool /*bNetworkIsConnected*/);
+
+	//Delegates called by JNI functions to bubble up underlying java work to the manager
+	FAndroidBackgroundDownload_OnWorkerStart AndroidBackgroundDownload_OnWorkerStart;
+	FAndroidBackgroundDownload_OnWorkerStop AndroidBackgroundDownload_OnWorkerStop;
+	FAndroidBackgroundDownload_OnProgress AndroidBackgroundDownload_OnProgress;
+	FAndroidBackgroundDownload_OnComplete AndroidBackgroundDownload_OnComplete;
+	FAndroidBackgroundDownload_OnAllComplete AndroidBackgroundDownload_OnAllComplete;
+	FAndroidBackgroundDownload_OnTickWorkerThread AndroidBackgroundDownload_OnTickWorkerThread;
+	FAndroidBackgroundDownload_OnNetworkChanged AndroidBackgroundDownload_OnNetworkChanged;
+};
+
 const FString FAndroidPlatformBackgroundHttpManager::BackgroundHTTPWorkID = TEXT("BackgroundHttpDownload");
 const FString FAndroidPlatformBackgroundHttpManager::AndroidBackgroundDownloadConfigRulesSettingKey = TEXT("AndroidBackgroundDownloadSetting");
 volatile int32 FAndroidPlatformBackgroundHttpManager::bHasManagerScheduledBGWork = false;
 
-FAndroidBackgroundDownloadDelegates::FAndroidBackgroundDownload_OnWorkerStart FAndroidBackgroundDownloadDelegates::AndroidBackgroundDownload_OnWorkerStart;
-FAndroidBackgroundDownloadDelegates::FAndroidBackgroundDownload_OnWorkerStop FAndroidBackgroundDownloadDelegates::AndroidBackgroundDownload_OnWorkerStop;
-FAndroidBackgroundDownloadDelegates::FAndroidBackgroundDownload_OnProgress FAndroidBackgroundDownloadDelegates::AndroidBackgroundDownload_OnProgress;
-FAndroidBackgroundDownloadDelegates::FAndroidBackgroundDownload_OnComplete FAndroidBackgroundDownloadDelegates::AndroidBackgroundDownload_OnComplete;
-FAndroidBackgroundDownloadDelegates::FAndroidBackgroundDownload_OnAllComplete FAndroidBackgroundDownloadDelegates::AndroidBackgroundDownload_OnAllComplete;
-FAndroidBackgroundDownloadDelegates::FAndroidBackgroundDownload_OnTickWorkerThread FAndroidBackgroundDownloadDelegates::AndroidBackgroundDownload_OnTickWorkerThread;
-
 const FString FAndroidNativeDownloadWorkerParameterKeys::DOWNLOAD_DESCRIPTION_LIST_KEY = TEXT("DownloadDescriptionList");
 const FString FAndroidNativeDownloadWorkerParameterKeys::DOWNLOAD_MAX_CONCURRENT_REQUESTS_KEY = TEXT("MaxConcurrentDownloadRequests");
 
@@ -80,7 +96,6 @@ const FString FAndroidNativeDownloadWorkerParameterKeys::NOTIFICATION_RESOURCE_S
 FAndroidPlatformBackgroundHttpManager::FJavaClassInfo FAndroidPlatformBackgroundHttpManager::JavaInfo = FAndroidPlatformBackgroundHttpManager::FJavaClassInfo();
 
 static bool bNetworkConnected = false;
-static volatile int32 bNetworkReconnected = false;
 
 static bool GetNetworkConnected()
 {
@@ -100,6 +115,9 @@ FAndroidPlatformBackgroundHttpManager::FAndroidPlatformBackgroundHttpManager()
 	, bIsModifyingPauseList(false)
 	, bIsModifyingResumeList(false)
 	, bIsModifyingCancelList(false)
+	, bHasPendingExpectedWorkStart(false)
+	, bShouldForceWorkerRequeue(false)
+	, bWorkerHadTerminalFinish(false)
 	, AndroidBackgroundHTTPManagerDefaultLocalizedText()
 {
 	if (JNIEnv* Env = FAndroidApplication::GetJavaEnv())
@@ -111,7 +129,6 @@ FAndroidPlatformBackgroundHttpManager::FAndroidPlatformBackgroundHttpManager()
 		}
 	}
 
-	bNetworkReconnected = false;
 	bNetworkConnected = GetNetworkConnected();
 }
 
@@ -145,11 +162,17 @@ bool FAndroidPlatformBackgroundHttpManager::HandleConfigRulesSettings()
 
 void FAndroidPlatformBackgroundHttpManager::Initialize()
 {	
-	Java_OnWorkerStopHandle = FAndroidBackgroundDownloadDelegates::AndroidBackgroundDownload_OnWorkerStop.AddRaw(this, &FAndroidPlatformBackgroundHttpManager::Java_OnWorkerStop);
-	Java_OnDownloadProgressHandle = FAndroidBackgroundDownloadDelegates::AndroidBackgroundDownload_OnProgress.AddRaw(this, &FAndroidPlatformBackgroundHttpManager::Java_OnDownloadProgress);
-	Java_OnDownloadCompleteHandle = FAndroidBackgroundDownloadDelegates::AndroidBackgroundDownload_OnComplete.AddRaw(this, &FAndroidPlatformBackgroundHttpManager::Java_OnDownloadComplete);
-	Java_OnAllDownloadsCompleteHandle = FAndroidBackgroundDownloadDelegates::AndroidBackgroundDownload_OnAllComplete.AddRaw(this, &FAndroidPlatformBackgroundHttpManager::Java_OnAllDownloadsComplete);
-	Java_OnTickHandle = FAndroidBackgroundDownloadDelegates::AndroidBackgroundDownload_OnTickWorkerThread.AddRaw(this, &FAndroidPlatformBackgroundHttpManager::Java_OnTick);
+	//Cancel any initial BackgroundHttpWork that is still running as this was kicked off by a previous run of the game
+	//and now that we are foregrounded we should wait to do that work until requested again.
+	FUEWorkManagerNativeWrapper::CancelBackgroundWork(BackgroundHTTPWorkID);
+
+	FAndroidBackgroundDownloadDelegates::AndroidBackgroundDownload_OnWorkerStart.BindSP(this, &FAndroidPlatformBackgroundHttpManager::Java_OnWorkerStart);
+	FAndroidBackgroundDownloadDelegates::AndroidBackgroundDownload_OnWorkerStop.BindSP(this, &FAndroidPlatformBackgroundHttpManager::Java_OnWorkerStop);
+	FAndroidBackgroundDownloadDelegates::AndroidBackgroundDownload_OnProgress.BindSP(this, &FAndroidPlatformBackgroundHttpManager::Java_OnDownloadProgress);
+	FAndroidBackgroundDownloadDelegates::AndroidBackgroundDownload_OnComplete.BindSP(this, &FAndroidPlatformBackgroundHttpManager::Java_OnDownloadComplete);
+	FAndroidBackgroundDownloadDelegates::AndroidBackgroundDownload_OnAllComplete.BindSP(this, &FAndroidPlatformBackgroundHttpManager::Java_OnAllDownloadsComplete);
+	FAndroidBackgroundDownloadDelegates::AndroidBackgroundDownload_OnTickWorkerThread.BindSP(this, &FAndroidPlatformBackgroundHttpManager::Java_OnTick);
+	FAndroidBackgroundDownloadDelegates::AndroidBackgroundDownload_OnNetworkChanged.BindSP(this, &FAndroidPlatformBackgroundHttpManager::Java_OnNetworkChanged);
 
 	AndroidBackgroundHTTPManagerDefaultLocalizedText.InitFromIniSettings("AndroidFetchBackgroundDownload");
 	FBackgroundHttpManagerImpl::Initialize();
@@ -157,12 +180,6 @@ void FAndroidPlatformBackgroundHttpManager::Initialize()
 
 void FAndroidPlatformBackgroundHttpManager::Shutdown()
 {
-	FAndroidBackgroundDownloadDelegates::AndroidBackgroundDownload_OnWorkerStop.Remove(Java_OnWorkerStopHandle);
-	FAndroidBackgroundDownloadDelegates::AndroidBackgroundDownload_OnProgress.Remove(Java_OnDownloadProgressHandle);
-	FAndroidBackgroundDownloadDelegates::AndroidBackgroundDownload_OnComplete.Remove(Java_OnDownloadCompleteHandle);
-	FAndroidBackgroundDownloadDelegates::AndroidBackgroundDownload_OnAllComplete.Remove(Java_OnAllDownloadsCompleteHandle);
-	FAndroidBackgroundDownloadDelegates::AndroidBackgroundDownload_OnTickWorkerThread.Remove(Java_OnTickHandle);
-
 	FBackgroundHttpManagerImpl::Shutdown();
 }
 
@@ -170,6 +187,8 @@ bool FAndroidPlatformBackgroundHttpManager::Tick(float DeltaTime)
 {
 	FBackgroundHttpManagerImpl::Tick(DeltaTime);
 
+	HandleJavaWorkerReconciliation();
+
 	HandlePendingCompletes();
 	HandleRequestsWaitingOnJavaThreadsafety();
 	UpdateRequestProgress();
@@ -181,16 +200,18 @@ void FAndroidPlatformBackgroundHttpManager::ActivatePendingRequests()
 {
 	UE_LOG(LogBackgroundHttpManager, VeryVerbose, TEXT("ActivatePendingRequests called"));
 
-	bool bRestoreActiveRequests = false;
-	if (FPlatformAtomics::AtomicRead(&bNetworkReconnected))
+	const bool bForceRequeueFlagged = FPlatformAtomics::InterlockedExchange(&bShouldForceWorkerRequeue, false);
+	bool bForceActiveRequestsToRequeue = false;
+	if (bForceRequeueFlagged)
 	{
 		FRWScopeLock ScopeLock(ActiveRequestLock, SLT_ReadOnly);
-		int32 ActiveRequestCount = ActiveRequests.Num();
-		UE_LOG(LogBackgroundHttpManager, Display, TEXT("Reconnected and found %d active requests to restore"), ActiveRequestCount);
-		bRestoreActiveRequests = (ActiveRequestCount > 0);
-		FPlatformAtomics::InterlockedExchange(&bNetworkReconnected, false);
+		{
+			int32 ActiveRequestCount = ActiveRequests.Num();
+			bForceActiveRequestsToRequeue = (ActiveRequestCount > 0);
+			UE_LOG(LogBackgroundHttpManager, Display, TEXT("Found %d active requests to force requeue."), ActiveRequestCount);
+		}
 	}
-
+	
 	bool bHasPendingRequests = false;
 	{
 		FRWScopeLock ScopeLock(PendingRequestLock, SLT_ReadOnly);
@@ -198,11 +219,13 @@ void FAndroidPlatformBackgroundHttpManager::ActivatePendingRequests()
 
 		UE_LOG(LogBackgroundHttpManager, VeryVerbose, TEXT("bHasPendingRequests: %d"), static_cast<int>(bHasPendingRequests));
 	}
-		
-	if (bHasPendingRequests || bRestoreActiveRequests)
+
+	//Queue up all PendingRequests and ActiveRequests with the underlying UEDownloadWorker layer
+	//Requires us to have PendingRequests or have flagged ActiveRequests for a re-queue
+	if (bHasPendingRequests || bForceActiveRequestsToRequeue)
 	{
 		UE_CLOG(bHasPendingRequests, LogBackgroundHttpManager, Display, TEXT("Found PendingRequests to activate"));
-		UE_CLOG(bRestoreActiveRequests, LogBackgroundHttpManager, Display, TEXT("Found ActiveRequests to restore"));
+		UE_CLOG(bForceActiveRequestsToRequeue, LogBackgroundHttpManager, Display, TEXT("Forcing Found ActiveRequests to requeue"));
 
 		const bool bIsOptional = false;
 
@@ -337,6 +360,9 @@ void FAndroidPlatformBackgroundHttpManager::ActivatePendingRequests()
 			{
 				//We have now scheduled some BG work so make sure we know its safe to call delegates
 				FPlatformAtomics::InterlockedExchange(&bHasManagerScheduledBGWork, true);
+
+				//Since we just scheduled new work, we can expect to get a call into OnWorkStart or OnWorkStop
+				FPlatformAtomics::InterlockedExchange(&bHasPendingExpectedWorkStart, true);
 			}
 			//We failed to schedule our background worker due to some error, so need to handle error case
 			else
@@ -353,6 +379,32 @@ void FAndroidPlatformBackgroundHttpManager::ActivatePendingRequests()
 	UE_LOG(LogBackgroundHttpManager, VeryVerbose, TEXT("ActivatePendingRequests complete"));
 }
 
+void FAndroidPlatformBackgroundHttpManager::HandleJavaWorkerReconciliation()
+{
+	//Check if we ended the last worker in a terminal fashion, and thus might need to fix up some underlying work requests
+	//if their state doesn't match this terminal state
+	const bool bWorkerHadTerminalFinishCopy = FPlatformAtomics::InterlockedExchange(&bWorkerHadTerminalFinish, false);
+	const bool bNewWorkerQueued = FPlatformAtomics::AtomicRead(&bHasPendingExpectedWorkStart);
+	if (bWorkerHadTerminalFinishCopy && !bNewWorkerQueued)
+	{
+		//if we are going to force requeue this worker that takes precedent here instead of forcing them to error out
+		const bool bShouldForceWorkerRequeueCopy = FPlatformAtomics::AtomicRead(&bShouldForceWorkerRequeue);
+		if (!bShouldForceWorkerRequeueCopy)
+		{
+			UE_LOG(LogBackgroundHttpManager, Display, TEXT("Completing all still active requests with an error as our worker ended with a terminal state while it was still downloading..."));
+			ForceCompleteAllUnderlyingJavaActiveRequestsWithError();
+		}
+		else
+		{
+			UE_LOG(LogBackgroundHttpManager, Display, TEXT("bWorkerHadTerminalFinish set but ignored as Worker is pending forced requeue"));
+		}
+	}
+	else if (bWorkerHadTerminalFinishCopy)
+	{
+		UE_LOG(LogBackgroundHttpManager, Display, TEXT("bWorkerHadTerminalFinish set but ignored as Worker is pending a new worker start"));
+	}
+}
+
 void FAndroidPlatformBackgroundHttpManager::HandlePendingCompletes()
 {
 	const bool bShouldProcessCompletes = FPlatformAtomics::InterlockedExchange(&bHasPendingCompletes, false);
@@ -634,7 +686,7 @@ void FAndroidPlatformBackgroundHttpManager::MarkUnderlyingJavaRequestAsCompleted
 	FAndroidBackgroundHttpRequestPtr AndroidRequest = StaticCastSharedPtr<FAndroidPlatformBackgroundHttpRequest>(Request);
 	if (ensureAlwaysMsgf((AndroidRequest.IsValid()), TEXT("Invalid Non-Android request!")))
 	{
-		return MarkUnderlyingJavaRequestAsCompleted(AndroidRequest);
+		return MarkUnderlyingJavaRequestAsCompleted(AndroidRequest, bSuccess);
 	}
 }
 
@@ -693,6 +745,25 @@ bool FAndroidPlatformBackgroundHttpManager::IsValidRequestToEnqueue(FAndroidBack
 	return (!bIsComplete && !bIsPaused);
 }
 
+void FAndroidPlatformBackgroundHttpManager::Java_OnWorkerStart(FString WorkID, jobject UnderlyingWorker)
+{
+	const bool bShouldRespondToDelegate = FPlatformAtomics::AtomicRead(&bHasManagerScheduledBGWork);
+	if (bShouldRespondToDelegate)
+	{
+		UE_LOG(LogBackgroundHttpManager, Log, TEXT("OnWorkerStart for %s"), *WorkID);
+
+		//With a new worker started flush this result if it hasn't been processed from the previous worker finishing
+		FPlatformAtomics::InterlockedExchange(&bWorkerHadTerminalFinish, false);
+
+		//Remove expected work start as workers have started
+		FPlatformAtomics::InterlockedExchange(&bHasPendingExpectedWorkStart, false);
+
+		//Remove forced requeue as another worker started and thus we can discard forcing a new worker start now
+		//that we have an active worker again
+		FPlatformAtomics::InterlockedExchange(&bShouldForceWorkerRequeue, false);
+	}
+}
+
 void FAndroidPlatformBackgroundHttpManager::Java_OnWorkerStop(FString WorkID, jobject UnderlyingWorker)
 {
 	const bool bShouldRespondToDelegate = FPlatformAtomics::AtomicRead(&bHasManagerScheduledBGWork);
@@ -701,10 +772,23 @@ void FAndroidPlatformBackgroundHttpManager::Java_OnWorkerStop(FString WorkID, jo
 		UE_LOG(LogBackgroundHttpManager, Log, TEXT("OnWorkerStop for %s"), *WorkID);
 
 		FUEWorkManagerNativeWrapper::EAndroidBackgroundWorkResult WorkResult = FUEWorkManagerNativeWrapper::GetWorkResultOnWorker(UnderlyingWorker);
-		if (WorkResult == FUEWorkManagerNativeWrapper::EAndroidBackgroundWorkResult::Failure)
+		if (WorkResult == FUEWorkManagerNativeWrapper::EAndroidBackgroundWorkResult::NotSet)
 		{
-			UE_LOG(LogBackgroundHttpManager, Log, TEXT("Worker Ended in Failure for %s. Completing ActiveRequests remaining with failure due to error."), *WorkID);
-			ForceCompleteAllUnderlyingJavaActiveRequestsWithError();
+			//Check if our work is stopping unexpectedly (we expect Stop to be called in the event of us scheduling new work)
+			//We know this is unexpected as we didn't set a WorkResult yet to denote the worker's state
+			bool bHasPendingStart = FPlatformAtomics::AtomicRead(&bHasPendingExpectedWorkStart);
+			if (!bHasPendingStart)
+			{
+				UE_LOG(LogBackgroundHttpManager, Warning, TEXT("Worker forced to end unexpectidly. Setting bShouldForceWorkerRequeue to force a new worker with all existing and pending work"));
+				FPlatformAtomics::InterlockedExchange(&bShouldForceWorkerRequeue, true);
+			}
+		}
+		else if (WorkResult != FUEWorkManagerNativeWrapper::EAndroidBackgroundWorkResult::Retry)
+		{
+			//WorkResult was terminal so we need to make sure we reconcile all work as worker is
+			//finished on next GameThread tick
+			UE_LOG(LogBackgroundHttpManager, Log, TEXT("Worker Ended with Terminal Result for %s"), *WorkID);
+			FPlatformAtomics::InterlockedExchange(&bWorkerHadTerminalFinish, true);
 		}
 	}
 }
@@ -730,8 +814,8 @@ void FAndroidPlatformBackgroundHttpManager::Java_OnAllDownloadsComplete(jobject
 			FUEWorkManagerNativeWrapper::SetWorkResultOnWorker(UnderlyingWorker, FUEWorkManagerNativeWrapper::EAndroidBackgroundWorkResult::Failure);
 		}
 
-		//Our underlying worker has finished all work, so mark all as complete and error if any still were waiting on complete
-		ForceCompleteAllUnderlyingJavaActiveRequestsWithError();
+		//Our underlying worker has finished all work, so this should be the same as a terminal finish
+		FPlatformAtomics::InterlockedExchange(&bWorkerHadTerminalFinish, true);
 	}
 }
 
@@ -828,6 +912,16 @@ void FAndroidPlatformBackgroundHttpManager::Java_OnTick(JNIEnv* Env, jobject Und
 		}
 	}
 }
+void FAndroidPlatformBackgroundHttpManager::Java_OnNetworkChanged(JNIEnv* Env, bool bNetworkIsConnected)
+{
+	if (bNetworkIsConnected)
+	{
+		UE_LOG(LogBackgroundHttpManager, Display, TEXT("Flagging worker for requeue due to network connection re-establishing."));
+		//On reconnecting to the network flag our worker to force requeue to make sure our work is continued
+		FPlatformAtomics::InterlockedExchange(&bShouldForceWorkerRequeue, true);
+	}
+	
+}
 
 void FAndroidPlatformBackgroundHttpManager::ForceCompleteAllUnderlyingJavaActiveRequestsWithError()
 {
@@ -1071,7 +1165,7 @@ JNI_METHOD void Java_com_epicgames_unreal_download_UEDownloadWorker_nativeAndroi
 	if (ensureAlwaysMsgf(bIsMatchingWorkID, TEXT("Unexpected WorkID %s sent back in OnWorkerStart! Ignoring un-related worker!"), *UEWorkID))
 	{
 		UE_LOG(LogBackgroundHttpManager, Display, TEXT("Called OnWorkerStart for %s"), *UEWorkID);
-		FAndroidBackgroundDownloadDelegates::AndroidBackgroundDownload_OnWorkerStart.Broadcast(UEWorkID, thiz);
+		FAndroidBackgroundDownloadDelegates::AndroidBackgroundDownload_OnWorkerStart.ExecuteIfBound(UEWorkID, thiz);
 	}
 }
 
@@ -1083,7 +1177,7 @@ JNI_METHOD void Java_com_epicgames_unreal_download_UEDownloadWorker_nativeAndroi
 	if (ensureAlwaysMsgf(bIsMatchingWorkID, TEXT("Unexpected WorkID %s sent back in OnWorkerStop! Ignoring un-related worker!"), *UEWorkID))
 	{
 		UE_LOG(LogBackgroundHttpManager, Display, TEXT("Called OnWorkerStop for %s"), *UEWorkID);
-		FAndroidBackgroundDownloadDelegates::AndroidBackgroundDownload_OnWorkerStop.Broadcast(UEWorkID, thiz);
+		FAndroidBackgroundDownloadDelegates::AndroidBackgroundDownload_OnWorkerStop.ExecuteIfBound(UEWorkID, thiz);
 	}
 }
 
@@ -1093,7 +1187,7 @@ JNI_METHOD void Java_com_epicgames_unreal_download_UEDownloadWorker_nativeAndroi
 	int64_t ConvertedBytesWrittenSinceLastCall = static_cast<uint64_t>(BytesWrittenSinceLastCall);
 	int64_t ConvertedTotalBytesWritten = static_cast<uint64_t>(TotalBytesWritten);
 
-	FAndroidBackgroundDownloadDelegates::AndroidBackgroundDownload_OnProgress.Broadcast(thiz, RequestID, ConvertedBytesWrittenSinceLastCall, ConvertedTotalBytesWritten);
+	FAndroidBackgroundDownloadDelegates::AndroidBackgroundDownload_OnProgress.ExecuteIfBound(thiz, RequestID, ConvertedBytesWrittenSinceLastCall, ConvertedTotalBytesWritten);
 }
 
 JNI_METHOD void Java_com_epicgames_unreal_download_UEDownloadWorker_nativeAndroidBackgroundDownloadOnComplete(JNIEnv* jenv, jobject thiz, jstring TaskID, jstring CompleteLocation, jboolean bWasSuccess)
@@ -1102,18 +1196,18 @@ JNI_METHOD void Java_com_epicgames_unreal_download_UEDownloadWorker_nativeAndroi
 	FString ConvertedCompleteLocation = FJavaHelper::FStringFromParam(jenv, CompleteLocation);
 	bool ConvertedbWasSuccess = static_cast<bool>(bWasSuccess);
 
-	FAndroidBackgroundDownloadDelegates::AndroidBackgroundDownload_OnComplete.Broadcast(thiz, RequestID, ConvertedCompleteLocation, ConvertedbWasSuccess);
+	FAndroidBackgroundDownloadDelegates::AndroidBackgroundDownload_OnComplete.ExecuteIfBound(thiz, RequestID, ConvertedCompleteLocation, ConvertedbWasSuccess);
 }
 
 JNI_METHOD void Java_com_epicgames_unreal_download_UEDownloadWorker_nativeAndroidBackgroundDownloadOnAllComplete(JNIEnv* jenv, jobject thiz, jboolean bDidAllRequestsSucceed)
 {
 	bool ConvertedbDidAllRequestsSucceed = static_cast<bool>(bDidAllRequestsSucceed);
-	FAndroidBackgroundDownloadDelegates::AndroidBackgroundDownload_OnAllComplete.Broadcast(thiz, bDidAllRequestsSucceed);
+	FAndroidBackgroundDownloadDelegates::AndroidBackgroundDownload_OnAllComplete.ExecuteIfBound(thiz, bDidAllRequestsSucceed);
 }
 
 JNI_METHOD void Java_com_epicgames_unreal_download_UEDownloadWorker_nativeAndroidBackgroundDownloadOnTick(JNIEnv* jenv, jobject thiz)
 {
-	FAndroidBackgroundDownloadDelegates::AndroidBackgroundDownload_OnTickWorkerThread.Broadcast(jenv, thiz);
+	FAndroidBackgroundDownloadDelegates::AndroidBackgroundDownload_OnTickWorkerThread.ExecuteIfBound(jenv, thiz);
 }
 
 //
@@ -1127,18 +1221,7 @@ JNI_METHOD void Java_com_epicgames_unreal_GameActivity_nativeAndroidBackgroundDo
 	if (bNewNetworkConnected != bNetworkConnected)
 	{
 		bNetworkConnected = bNewNetworkConnected;
-
-		if (bNetworkConnected)
-		{
-			// reconnected !
-			//FPlatformMisc::LowLevelOutputDebugString(TEXT("nativeAFBD_NetworkChanged: Reconnected!"));
-			FPlatformAtomics::InterlockedExchange(&bNetworkReconnected, true);
-		}
-		else
-		{
-			// disconnected !
-			//FPlatformMisc::LowLevelOutputDebugString(TEXT("nativeAFBD_NetworkChanged: Disconnected!"));
-		}
+		FAndroidBackgroundDownloadDelegates::AndroidBackgroundDownload_OnNetworkChanged.ExecuteIfBound(jenv, bNetworkConnected);
 	}
 }
 
diff --git a/Engine/Plugins/Online/Android/AndroidFetchBackgroundDownload/Source/AndroidFetchBackgroundDownload/Public/AndroidPlatformBackgroundHttpManager.h b/Engine/Plugins/Online/Android/AndroidFetchBackgroundDownload/Source/AndroidFetchBackgroundDownload/Public/AndroidPlatformBackgroundHttpManager.h
--- a/Engine/Plugins/Online/Android/AndroidFetchBackgroundDownload/Source/AndroidFetchBackgroundDownload/Public/AndroidPlatformBackgroundHttpManager.h
+++ b/Engine/Plugins/Online/Android/AndroidFetchBackgroundDownload/Source/AndroidFetchBackgroundDownload/Public/AndroidPlatformBackgroundHttpManager.h
@@ -64,15 +64,19 @@ protected:
 	FAndroidBackgroundHttpRequestPtr FindRequestByID(FString RequestID);
 	void HandlePendingCompletes();
 
+	void HandleJavaWorkerReconciliation();
+
 	const FString GetFullFileNameForDownloadDescriptionList() const;
 	const FString GetBaseFileNameForDownloadDescriptionListWithAppendedInt(int IntToAppend) const;
 
 	//Handlers for our download progressing in the underlying java implementation so that we can bubble it up to UE code.
+	void Java_OnWorkerStart(FString WorkID, jobject UnderlyingWorker);
 	void Java_OnWorkerStop(FString WorkID, jobject UnderlyingWorker);
 	void Java_OnDownloadProgress(jobject UnderlyingWorker, FString RequestID, int64_t BytesWrittenSinceLastCall, int64_t TotalBytesWritten);
 	void Java_OnDownloadComplete(jobject UnderlyingWorker, FString RequestID, FString CompleteLocation, bool bWasSuccess);
 	void Java_OnAllDownloadsComplete(jobject UnderlyingWorker, bool bDidAllRequestsSucceed);
 	void Java_OnTick(JNIEnv* Env, jobject UnderlyingWorker);
+	void Java_OnNetworkChanged(JNIEnv* Env, bool bNetworkIsConnected);
 
 	//Helper function that completes any un-completed requests in the underlying java tracking layer so that they will
 	//complete on the next Tick. Mostly used to help when the underlying java error hits a failure. Ensures when any
@@ -92,12 +96,6 @@ protected:
 	bool IsValidRequestToEnqueue(FBackgroundHttpRequestPtr Request);
 	bool IsValidRequestToEnqueue(FAndroidBackgroundHttpRequestPtr Request);
 
-	FDelegateHandle Java_OnWorkerStopHandle;
-	FDelegateHandle Java_OnDownloadProgressHandle;
-	FDelegateHandle Java_OnDownloadCompleteHandle;
-	FDelegateHandle Java_OnAllDownloadsCompleteHandle;
-	FDelegateHandle Java_OnTickHandle;
-
 	//Array used to store Pause/Resume/Cancel requests in a thread-safe non-locking way. This way we can utilize the _Java lists in our Java_OnTick
 	//without worrying about blocking the java thread
 	TArray<FString> RequestsToPauseByID_GT;
@@ -117,6 +115,17 @@ protected:
 	volatile int32 bIsModifyingResumeList;
 	volatile int32 bIsModifyingCancelList;
 
+	//Tracks if we are expecting a worker to start
+	volatile int32 bHasPendingExpectedWorkStart;
+
+	//If for some reason we should force requeuing all active work even if there are no current PendingDownloadRequests
+	volatile int32 bShouldForceWorkerRequeue;
+	
+	//If we need to reconcile the state of our underlying requests as
+	//we have terminally finished our work and want to make sure all work is
+	//correctly set to finish.
+	volatile int32 bWorkerHadTerminalFinish;
+
 	//Rechecks any _GT lists to try and move them to _Java lists if its safe to do so
 	void HandleRequestsWaitingOnJavaThreadsafety();
 
@@ -226,24 +235,5 @@ public:
 	static const FString NOTIFICATION_RESOURCE_SMALL_ICON_PACKAGE;
 
 };
-//Call backs called by the bellow FBackgroundURLSessionHandler so higher-level systems can respond to task updates.
-class FAndroidBackgroundDownloadDelegates
-{
-public:
-	DECLARE_MULTICAST_DELEGATE_TwoParams(FAndroidBackgroundDownload_OnWorkerStart, FString /*WorkID*/, jobject /*UnderlyingWorker*/);
-	DECLARE_MULTICAST_DELEGATE_TwoParams(FAndroidBackgroundDownload_OnWorkerStop, FString /*WorkID*/, jobject /*UnderlyingWorker*/);
-	DECLARE_MULTICAST_DELEGATE_FourParams(FAndroidBackgroundDownload_OnProgress, jobject /*UnderlyingWorker*/, FString /*RequestID*/, int64_t /*BytesWrittenSinceLastCall*/, int64_t /*TotalBytesWritten*/);
-	DECLARE_MULTICAST_DELEGATE_FourParams(FAndroidBackgroundDownload_OnComplete, jobject /*UnderlyingWorker*/, FString /*RequestID*/, FString /*CompleteLocation*/, bool /*bWasSuccess*/);
-	DECLARE_MULTICAST_DELEGATE_TwoParams(FAndroidBackgroundDownload_OnAllComplete, jobject /*UnderlyingWorker*/, bool /*bDidAllRequestsSucceed*/);
-	DECLARE_MULTICAST_DELEGATE_TwoParams(FAndroidBackgroundDownload_OnTickWorkerThread, JNIEnv*, jobject /*UnderlyingWorker*/);
-
-	//Delegates called by JNI functions to bubble up underlying java work to the manager
-	static FAndroidBackgroundDownload_OnWorkerStart AndroidBackgroundDownload_OnWorkerStart;
-	static FAndroidBackgroundDownload_OnWorkerStop AndroidBackgroundDownload_OnWorkerStop;
-	static FAndroidBackgroundDownload_OnProgress AndroidBackgroundDownload_OnProgress;
-	static FAndroidBackgroundDownload_OnComplete AndroidBackgroundDownload_OnComplete;
-	static FAndroidBackgroundDownload_OnAllComplete AndroidBackgroundDownload_OnAllComplete;
-	static FAndroidBackgroundDownload_OnTickWorkerThread AndroidBackgroundDownload_OnTickWorkerThread;
-};
 
 typedef TSharedPtr<FAndroidPlatformBackgroundHttpManager, ESPMode::ThreadSafe> FAndroidPlatformBackgroundHttpManagerPtr;
\ No newline at end of file
diff --git a/Engine/Source/Runtime/Online/BackgroundHTTP/Public/Interfaces/IBackgroundHttpManager.h b/Engine/Source/Runtime/Online/BackgroundHTTP/Public/Interfaces/IBackgroundHttpManager.h
--- a/Engine/Source/Runtime/Online/BackgroundHTTP/Public/Interfaces/IBackgroundHttpManager.h
+++ b/Engine/Source/Runtime/Online/BackgroundHTTP/Public/Interfaces/IBackgroundHttpManager.h
@@ -6,9 +6,12 @@
 
 typedef TSharedPtr<class IBackgroundHttpManager, ESPMode::ThreadSafe> FBackgroundHttpManagerPtr;
 
-class BACKGROUNDHTTP_API IBackgroundHttpManager
+class BACKGROUNDHTTP_API IBackgroundHttpManager : public TSharedFromThis<IBackgroundHttpManager>
 {
 public:
+
+	virtual ~IBackgroundHttpManager() = default;
+
 	/**
 	 * Initialize
 	 */
