Android Background Download Bug Fixes:
- Fixed AndroidPlatformBackgroundHttpManager incorrectly not re-sending paused downloads in payloads. Fixed stall when all downloads were paused and the worker was re-queued (such as with cellular download warnings).
- Fixed pause state not being sent to underlying UEDownloadWorker through DownloadDescription
- Logging improvements
#tests Android Dev Client

diff --git a/Engine/Plugins/Online/Android/AndroidFetchBackgroundDownload/Source/AndroidFetchBackgroundDownload/Java/download/UEDownloadWorker.java b/Engine/Plugins/Online/Android/AndroidFetchBackgroundDownload/Source/AndroidFetchBackgroundDownload/Java/download/UEDownloadWorker.java
--- a/Engine/Plugins/Online/Android/AndroidFetchBackgroundDownload/Source/AndroidFetchBackgroundDownload/Java/download/UEDownloadWorker.java
+++ b/Engine/Plugins/Online/Android/AndroidFetchBackgroundDownload/Source/AndroidFetchBackgroundDownload/Java/download/UEDownloadWorker.java
@@ -135,7 +135,7 @@ public class UEDownloadWorker extends UEWorker implements DownloadProgressListen
        //Have to have parsed some DownloadDescriptions to have any meaningful work to do
 		if ((QueueDescription == null) || (QueueDescription.DownloadDescriptions.size() == 0))
 		{
-			Log.error("Invalid QueueDescription! No DownloadDescription list for queued UEDownloadWorker!");
+			Log.error("Invalid QueueDescription! No DownloadDescription list for queued UEDownloadWorker! Worker InputData:" + getInputData());
 			SetWorkResult_Failure();
 			return;
 		}
@@ -175,7 +175,6 @@ public class UEDownloadWorker extends UEWorker implements DownloadProgressListen
 		if (!bReceivedResult && bHasEnqueueHappened && !bForceStopped)
 		{
 			mFetchManager.RequestGroupProgressUpdate(QueueDescription.DownloadGroupID,  this);
-			mFetchManager.RequestCheckDownloadsStillActive(this);
 
 			nativeAndroidBackgroundDownloadOnTick();
 		}
diff --git a/Engine/Plugins/Online/Android/AndroidFetchBackgroundDownload/Source/AndroidFetchBackgroundDownload/Java/download/datastructs/DownloadDescription.java b/Engine/Plugins/Online/Android/AndroidFetchBackgroundDownload/Source/AndroidFetchBackgroundDownload/Java/download/datastructs/DownloadDescription.java
--- a/Engine/Plugins/Online/Android/AndroidFetchBackgroundDownload/Source/AndroidFetchBackgroundDownload/Java/download/datastructs/DownloadDescription.java
+++ b/Engine/Plugins/Online/Android/AndroidFetchBackgroundDownload/Source/AndroidFetchBackgroundDownload/Java/download/datastructs/DownloadDescription.java
@@ -36,14 +36,22 @@ public class DownloadDescription
 	public boolean bHasCompleted = false;
 
 	//
+	//Serliazed AND state Members -- Initial state is pulled from WorkerParameters but further
+	//changes are made to track download information
+	//
+
+	//If our download is currently paused.
+	public boolean bIsPaused = false;
+
 	//Non-serialized members -- Used to track download information
 	//
 	public int CurrentRetryCount = 0;
 	public int CachedFetchID = 0;
-	public boolean bIsPaused = false;
+	
 	public boolean bIsCancelled = false;
-	public long PreviousDownloadedBytes = 0;
+
 	public int PreviousDownloadPercent = 0;
+	public long PreviousDownloadedBytes = 0;
 	public long TotalBytesNeeded = 0;
 	public long TotalDownloadedBytes = 0;
 	public DownloadProgressListener ProgressListener = null;
@@ -215,6 +223,12 @@ public class DownloadDescription
 			{
 				ReturnedDescription.bHasCompleted = jObject.optBoolean(bHasCompletedKey);
 			}
+
+			if (false == jObject.isNull(bIsPausedKey))
+			{
+				ReturnedDescription.bIsPaused = jObject.optBoolean(bIsPausedKey);
+			}
+
 		}
 		catch( JSONException e )
 		{
@@ -313,4 +327,5 @@ public class DownloadDescription
 	private static final String RequestPriorityKey = "RequestPriority";
 	private static final String GroupIDKey = "GroupId";
 	private static final String bHasCompletedKey = "bHasCompleted";
+	private static final String bIsPausedKey = "bIsPaused";
 }
\ No newline at end of file
diff --git a/Engine/Plugins/Online/Android/AndroidFetchBackgroundDownload/Source/AndroidFetchBackgroundDownload/Java/download/fetch/FetchManager.java b/Engine/Plugins/Online/Android/AndroidFetchBackgroundDownload/Source/AndroidFetchBackgroundDownload/Java/download/fetch/FetchManager.java
--- a/Engine/Plugins/Online/Android/AndroidFetchBackgroundDownload/Source/AndroidFetchBackgroundDownload/Java/download/fetch/FetchManager.java
+++ b/Engine/Plugins/Online/Android/AndroidFetchBackgroundDownload/Source/AndroidFetchBackgroundDownload/Java/download/fetch/FetchManager.java
@@ -320,8 +320,8 @@ public class FetchManager implements FetchDownloadProgressOwner, FetchEnqueueRes
 		NewDescription.TotalBytesNeeded = OldDescription.TotalBytesNeeded;
 		NewDescription.TotalDownloadedBytes = OldDescription.TotalDownloadedBytes;
 
-		// Purposefully don't copy bIsPaused or bIsCancelled as we want that to come from the new description in case 
-		// this has been set on C++ side but not yet received by our Worker before new work was queued. C++ state should be definitive
+		// Purposefully don't copy bIsCancelled as this isn't tracked on the C++ side and if we are asking for this to be re-queued then
+		// we no longer want this download cancelled and should redo it
 
 		//Purposefully don't copy DownloadProgressListener as we set this based on the new QueueDescription already and likely
 		//don't want to keep the old one (although in current code they match).
diff --git a/Engine/Plugins/Online/Android/AndroidFetchBackgroundDownload/Source/AndroidFetchBackgroundDownload/Private/AndroidPlatformBackgroundHttpManager.cpp b/Engine/Plugins/Online/Android/AndroidFetchBackgroundDownload/Source/AndroidFetchBackgroundDownload/Private/AndroidPlatformBackgroundHttpManager.cpp
--- a/Engine/Plugins/Online/Android/AndroidFetchBackgroundDownload/Source/AndroidFetchBackgroundDownload/Private/AndroidPlatformBackgroundHttpManager.cpp
+++ b/Engine/Plugins/Online/Android/AndroidFetchBackgroundDownload/Source/AndroidFetchBackgroundDownload/Private/AndroidPlatformBackgroundHttpManager.cpp
@@ -302,19 +302,18 @@ void FAndroidPlatformBackgroundHttpManager::ActivatePendingRequests()
 			UE_LOG(LogBackgroundHttpManager, Display, TEXT("Adding %d requests to UEDownloadableWorker's DownloadDescription list for active work to be done."), RequestsToSendToJavaLayer.Num());
 			for (FAndroidBackgroundHttpRequestPtr& Request : RequestsToSendToJavaLayer)
 			{
-				if (IsValidRequestToEnqueue(Request))
-				{
-					FScopedJavaObject<jstring> JavaJSONString = FJavaHelper::ToJavaString(Env, Request->ToJSon());
-					FScopedJavaObject<jobject> Description = NewScopedJavaObject(Env, Env->CallStaticObjectMethod(FAndroidPlatformBackgroundHttpManager::JavaInfo.DownloadDescriptionClass, FAndroidPlatformBackgroundHttpManager::JavaInfo.CreateDownloadDescriptionFromJsonMethod, *JavaJSONString));
+				const FString RequestJSONString = Request->ToJSon();
+				FScopedJavaObject<jstring> JavaJSONString = FJavaHelper::ToJavaString(Env, RequestJSONString);
+				FScopedJavaObject<jobject> Description = NewScopedJavaObject(Env, Env->CallStaticObjectMethod(FAndroidPlatformBackgroundHttpManager::JavaInfo.DownloadDescriptionClass, FAndroidPlatformBackgroundHttpManager::JavaInfo.CreateDownloadDescriptionFromJsonMethod, *JavaJSONString));
+				UE_LOG(LogBackgroundHttpManager, VeryVerbose, TEXT("Generated JSON for request %s. || %s"), *Request->GetRequestID(), *RequestJSONString);
 
-					bool bAddedRequestAsDescription = FJavaWrapper::CallBooleanMethod(Env, *DescriptionArray, ArrayPutMethod, *Description);
-					if (!bAddedRequestAsDescription)
-					{
-						ensureAlwaysMsgf(false, TEXT("Failed to create and add valid DownloadDescription for request %s"), *Request->GetRequestID());
-						//We failed to add this request to our download worker so just mark it as completed so it can send a failure completion
-						MarkUnderlyingJavaRequestAsCompleted(Request, false);
-					}
-				}				
+				bool bAddedRequestAsDescription = FJavaWrapper::CallBooleanMethod(Env, *DescriptionArray, ArrayPutMethod, *Description);
+				if (!bAddedRequestAsDescription)
+				{
+					ensureAlwaysMsgf(false, TEXT("Failed to create and add valid DownloadDescription for request %s"), *Request->GetRequestID());
+					//We failed to add this request to our download worker so just mark it as completed so it can send a failure completion
+					MarkUnderlyingJavaRequestAsCompleted(Request, false);
+				}
 			}
 				
 			//Call JNI function that saves our passed in ArrayList<DownloadDescription> to a file
@@ -352,7 +351,7 @@ void FAndroidPlatformBackgroundHttpManager::ActivatePendingRequests()
 				FText UpdatedContentText = FText::Format(AndroidBackgroundHTTPManagerDefaultLocalizedText.DefaultNotificationText_Content.GetText(), Arguments);
 				WorkParams.AddDataToWorkerParameters(FAndroidNativeDownloadWorkerParameterKeys::NOTIFICATION_CONTENT_TEXT_KEY, UpdatedContentText);
 
-				UE_LOG(LogBackgroundHttpManager, Display, TEXT("Attempting to schedule UEDownloadableWorker for background work"));
+				UE_LOG(LogBackgroundHttpManager, Display, TEXT("Attempting to schedule UEDownloadableWorker for background work. DownloadDescriptionList FileName:%s"), *FileNameForDownloadDescList);
 				bDidSuccessfullyScheduleWork  = FUEWorkManagerNativeWrapper::ScheduleBackgroundWork(BackgroundHTTPWorkID, WorkParams);
 			}
 			
@@ -478,7 +477,8 @@ void FAndroidPlatformBackgroundHttpManager::PauseRequest(FBackgroundHttpRequestP
 	if (ensureAlwaysMsgf(IsInGameThread(), TEXT("Should only ever call PauseRequest from GameThread! Can not pause!")))
 	{
 		const FString RequestID = Request->GetRequestID();
-		
+		UE_LOG(LogBackgroundHttpManager, Verbose, TEXT("PauseRequest called on %s"), *RequestID);
+
 		//Flag our request as paused
 		FAndroidBackgroundHttpRequestPtr AndroidRequest = StaticCastSharedPtr<FAndroidPlatformBackgroundHttpRequest>(Request);
 		if (ensureAlwaysMsgf(AndroidRequest.IsValid(), TEXT("Invalid request in Active Requests!")))
@@ -508,6 +508,7 @@ void FAndroidPlatformBackgroundHttpManager::ResumeRequest(FBackgroundHttpRequest
 	if (ensureAlwaysMsgf(IsInGameThread(), TEXT("Should only ever call ResumeRequest from GameThread! Can not Resume!")))
 	{
 		const FString RequestID = Request->GetRequestID();
+		UE_LOG(LogBackgroundHttpManager, Verbose, TEXT("ResumeRequest called on %s"), *RequestID);
 
 		//Flag our request as resumed
 		FAndroidBackgroundHttpRequestPtr AndroidRequest = StaticCastSharedPtr<FAndroidPlatformBackgroundHttpRequest>(Request);
@@ -538,6 +539,7 @@ void FAndroidPlatformBackgroundHttpManager::CancelRequest(FBackgroundHttpRequest
 	if (ensureAlwaysMsgf(IsInGameThread(), TEXT("Should only ever call CancelRequest from GameThread! Can not Cancel!")))
 	{
 		const FString RequestID = Request->GetRequestID();
+		UE_LOG(LogBackgroundHttpManager, Verbose, TEXT("CancelRequest called on %s"), *RequestID);
 
 		//Go ahead and remove our request since its cancelled
 		RemoveRequest(Request);
@@ -726,25 +728,6 @@ bool FAndroidPlatformBackgroundHttpManager::HasUnderlyingJavaCompletedRequest(FA
 	return false;
 }
 
-bool FAndroidPlatformBackgroundHttpManager::IsValidRequestToEnqueue(FBackgroundHttpRequestPtr Request)
-{
-	FAndroidBackgroundHttpRequestPtr AndroidRequest = StaticCastSharedPtr<FAndroidPlatformBackgroundHttpRequest>(Request);
-	if (ensureAlwaysMsgf((AndroidRequest.IsValid()), TEXT("Invalid Non-Android request checked for valid request! Returning false")))
-	{
-		return IsValidRequestToEnqueue(AndroidRequest);
-	}
-
-	return false;
-}
-
-bool FAndroidPlatformBackgroundHttpManager::IsValidRequestToEnqueue(FAndroidBackgroundHttpRequestPtr Request)
-{
-	const bool bIsComplete = HasUnderlyingJavaCompletedRequest(Request);
-	const bool bIsPaused = FPlatformAtomics::AtomicRead(&(Request->bIsPaused));
-
-	return (!bIsComplete && !bIsPaused);
-}
-
 void FAndroidPlatformBackgroundHttpManager::Java_OnWorkerStart(FString WorkID, jobject UnderlyingWorker)
 {
 	const bool bShouldRespondToDelegate = FPlatformAtomics::AtomicRead(&bHasManagerScheduledBGWork);
diff --git a/Engine/Plugins/Online/Android/AndroidFetchBackgroundDownload/Source/AndroidFetchBackgroundDownload/Private/AndroidPlatformBackgroundHttpRequest.cpp b/Engine/Plugins/Online/Android/AndroidFetchBackgroundDownload/Source/AndroidFetchBackgroundDownload/Private/AndroidPlatformBackgroundHttpRequest.cpp
--- a/Engine/Plugins/Online/Android/AndroidFetchBackgroundDownload/Source/AndroidFetchBackgroundDownload/Private/AndroidPlatformBackgroundHttpRequest.cpp
+++ b/Engine/Plugins/Online/Android/AndroidFetchBackgroundDownload/Source/AndroidFetchBackgroundDownload/Private/AndroidPlatformBackgroundHttpRequest.cpp
@@ -19,10 +19,11 @@ const FString FAndroidPlatformBackgroundHttpRequest::IndividualURLRetryCountKey
 const FString FAndroidPlatformBackgroundHttpRequest::RequestPriorityKey = "RequestPriority";
 const FString FAndroidPlatformBackgroundHttpRequest::GroupIDKey = "GroupId";
 const FString FAndroidPlatformBackgroundHttpRequest::bHasCompletedKey = "bHasCompleted";
+const FString FAndroidPlatformBackgroundHttpRequest::bIsPausedKey = "bIsPaused";
 
 FAndroidPlatformBackgroundHttpRequest::FAndroidPlatformBackgroundHttpRequest()
-	: bIsCompleted(false)
-	, bIsPaused(false)
+	: bIsPaused(false)
+	, bIsCompleted(false)
 	, DownloadProgress(0)
 	, DownloadProgressSinceLastUpdateSent(0)
 {
@@ -62,6 +63,11 @@ FString FAndroidPlatformBackgroundHttpRequest::ToJSon() const
 		const FString HasCompletedString = bIsCompletedCopy ? TEXT("true") : TEXT("false");
 		Writer->WriteValue(bHasCompletedKey, HasCompletedString);
 
+		//Write this bool as either true/false so it can be parsed in JSON as a java Boolean object
+		const bool bIsPausedCopy = FPlatformAtomics::AtomicRead(&bIsPaused);
+		const FString IsPausedString = bIsPausedCopy ? TEXT("true") : TEXT("false");
+		Writer->WriteValue(bIsPausedKey, bIsPausedCopy);
+
 		//TODO: The intent of this key is to allow multiple download notifications to be active at once and this group would mean all notifications with the same key
 		//are lumped together. For now everything is just expected to be in the same group, but if we want to support this we can implement a more meaningful groupID here.
 		Writer->WriteValue(GroupIDKey, 0);
@@ -82,7 +88,7 @@ FString FAndroidPlatformBackgroundHttpRequest::ToJSon() const
 	Writer->WriteObjectEnd();
 	Writer->Close();
 	
-	return JsonOutput;
+	return MoveTemp(JsonOutput);
 }
 
 FString FAndroidPlatformBackgroundHttpRequest::GetDestinationLocation() const
diff --git a/Engine/Plugins/Online/Android/AndroidFetchBackgroundDownload/Source/AndroidFetchBackgroundDownload/Public/AndroidPlatformBackgroundHttpManager.h b/Engine/Plugins/Online/Android/AndroidFetchBackgroundDownload/Source/AndroidFetchBackgroundDownload/Public/AndroidPlatformBackgroundHttpManager.h
--- a/Engine/Plugins/Online/Android/AndroidFetchBackgroundDownload/Source/AndroidFetchBackgroundDownload/Public/AndroidPlatformBackgroundHttpManager.h
+++ b/Engine/Plugins/Online/Android/AndroidFetchBackgroundDownload/Source/AndroidFetchBackgroundDownload/Public/AndroidPlatformBackgroundHttpManager.h
@@ -92,10 +92,6 @@ protected:
 	void MarkUnderlyingJavaRequestAsCompleted(FBackgroundHttpRequestPtr Request, bool bSuccess = true );
 	void MarkUnderlyingJavaRequestAsCompleted(FAndroidBackgroundHttpRequestPtr Request, bool bSuccess = true);
 
-	//returns true if this request is a valid request to send through to our underlying Java worker. Makes sure the request is not completed or paused
-	bool IsValidRequestToEnqueue(FBackgroundHttpRequestPtr Request);
-	bool IsValidRequestToEnqueue(FAndroidBackgroundHttpRequestPtr Request);
-
 	//Array used to store Pause/Resume/Cancel requests in a thread-safe non-locking way. This way we can utilize the _Java lists in our Java_OnTick
 	//without worrying about blocking the java thread
 	TArray<FString> RequestsToPauseByID_GT;
diff --git a/Engine/Plugins/Online/Android/AndroidFetchBackgroundDownload/Source/AndroidFetchBackgroundDownload/Public/AndroidPlatformBackgroundHttpRequest.h b/Engine/Plugins/Online/Android/AndroidFetchBackgroundDownload/Source/AndroidFetchBackgroundDownload/Public/AndroidPlatformBackgroundHttpRequest.h
--- a/Engine/Plugins/Online/Android/AndroidFetchBackgroundDownload/Source/AndroidFetchBackgroundDownload/Public/AndroidPlatformBackgroundHttpRequest.h
+++ b/Engine/Plugins/Online/Android/AndroidFetchBackgroundDownload/Source/AndroidFetchBackgroundDownload/Public/AndroidPlatformBackgroundHttpRequest.h
@@ -41,6 +41,7 @@ private:
 	static const FString RequestPriorityKey;
 	static const FString GroupIDKey;
 	static const FString bHasCompletedKey;
+	static const FString bIsPausedKey;
 
 private:
 	volatile int32 bIsPaused;
