Don't complete all the requests in the background downloader when pausing.

diff --git a/Engine/Plugins/Online/Android/AndroidFetchBackgroundDownload/Source/AndroidFetchBackgroundDownload/Java/download/UEDownloadWorker.java b/Engine/Plugins/Online/Android/AndroidFetchBackgroundDownload/Source/AndroidFetchBackgroundDownload/Java/download/UEDownloadWorker.java
--- a/Engine/Plugins/Online/Android/AndroidFetchBackgroundDownload/Source/AndroidFetchBackgroundDownload/Java/download/UEDownloadWorker.java
+++ b/Engine/Plugins/Online/Android/AndroidFetchBackgroundDownload/Source/AndroidFetchBackgroundDownload/Java/download/UEDownloadWorker.java
@@ -176,6 +176,9 @@ public class UEDownloadWorker extends UEWorker implements DownloadProgressListen
 		{
 			mFetchManager.RequestGroupProgressUpdate(QueueDescription.DownloadGroupID,  this);
 
+			//Keeping the code path to insert a watch dog later on
+			//mFetchManager.RequestCheckDownloadsStillActive(this);
+
 			nativeAndroidBackgroundDownloadOnTick();
 		}
 	}
diff --git a/Engine/Plugins/Online/Android/AndroidFetchBackgroundDownload/Source/AndroidFetchBackgroundDownload/Java/download/fetch/FetchManager.java b/Engine/Plugins/Online/Android/AndroidFetchBackgroundDownload/Source/AndroidFetchBackgroundDownload/Java/download/fetch/FetchManager.java
--- a/Engine/Plugins/Online/Android/AndroidFetchBackgroundDownload/Source/AndroidFetchBackgroundDownload/Java/download/fetch/FetchManager.java
+++ b/Engine/Plugins/Online/Android/AndroidFetchBackgroundDownload/Source/AndroidFetchBackgroundDownload/Java/download/fetch/FetchManager.java
@@ -869,10 +869,7 @@ public class FetchManager implements FetchDownloadProgressOwner, FetchEnqueueRes
 	{
 		if (!bHasActiveDownloads)
 		{
-			Log.error("No active downloads in FetchInstance, but still waiting on downloads to finish! Acting as if all downloads completed.");
-			//Send a failure download complete notification since we have an unexpected error where we are still trying to do work
-			//without any active downloads in our FetchManager
-			SendAllDownloadsCompleteNotification(ProgressListener, true);
+			Log.debug("No active downloads in FetchInstance, but still waiting on downloads to finish!");
 		}
 	}
 
