Change AndroidBackgroundHttp notification to be silent to avoid vibrating/playing sound when queueing downloads.
#tests Android Dev Client

diff --git a/Engine/Plugins/Online/Android/AndroidFetchBackgroundDownload/Source/AndroidFetchBackgroundDownload/Java/download/UEDownloadWorker.java b/Engine/Plugins/Online/Android/AndroidFetchBackgroundDownload/Source/AndroidFetchBackgroundDownload/Java/download/UEDownloadWorker.java
--- a/Engine/Plugins/Online/Android/AndroidFetchBackgroundDownload/Source/AndroidFetchBackgroundDownload/Java/download/UEDownloadWorker.java
+++ b/Engine/Plugins/Online/Android/AndroidFetchBackgroundDownload/Source/AndroidFetchBackgroundDownload/Java/download/UEDownloadWorker.java
@@ -297,6 +297,7 @@ public class UEDownloadWorker extends UEWorker implements DownloadProgressListen
 			.setOnlyAlertOnce (true)
 			.setSmallIcon(Description.SmallIconResourceID)
 			.addAction(Description.CancelIconResourceID, Description.CancelText, CancelIntent)
+			.setNotificationSilent()
 			.build();
 
 		return notification;
@@ -328,6 +329,7 @@ public class UEDownloadWorker extends UEWorker implements DownloadProgressListen
 			.setOnlyAlertOnce (true)
 			.setSmallIcon(Description.SmallIconResourceID)
 			.addAction(Description.CancelIconResourceID, Description.CancelText, CancelIntent)
+			.setNotificationSilent()
 			.build();
 
 		return notification;
