When Reconcileing items, the first item could be completed leading to the worker shutting down.
- Now only check for completed downloads on the last item
- Retry and other actions are the same as before

diff --git a/Engine/Plugins/Online/Android/AndroidFetchBackgroundDownload/Source/AndroidFetchBackgroundDownload/Java/download/fetch/FetchManager.java b/Engine/Plugins/Online/Android/AndroidFetchBackgroundDownload/Source/AndroidFetchBackgroundDownload/Java/download/fetch/FetchManager.java
--- a/Engine/Plugins/Online/Android/AndroidFetchBackgroundDownload/Source/AndroidFetchBackgroundDownload/Java/download/fetch/FetchManager.java
+++ b/Engine/Plugins/Online/Android/AndroidFetchBackgroundDownload/Source/AndroidFetchBackgroundDownload/Java/download/fetch/FetchManager.java
@@ -127,15 +127,16 @@ public class FetchManager implements FetchDownloadProgressOwner, FetchEnqueueRes
 							
 			String RequestIDKey = NewDownloadDescription.RequestID;
 			//Don't even have an entry for this RequestID in the old list, so it's completely new
+			boolean LastRequest = (NumNewDownloadDescriptions-1) == DescriptionIndex;
 			if (!TempRequestedDownloadsCopy.containsKey(RequestIDKey))
 			{
-				QueueNewDownloadDescription(NewDownloadDescription);
+				QueueNewDownloadDescription(NewDownloadDescription, LastRequest);
 			}
 			//Need to update our DownloadDescription
 			else
 			{
 				DownloadDescription ActiveDescription = TempRequestedDownloadsCopy.get(RequestIDKey);
-				HandleChangedDownloadDescription(ActiveDescription, NewDownloadDescription);
+				HandleChangedDownloadDescription(ActiveDescription, NewDownloadDescription, LastRequest);
 			}
 		}
 	}
@@ -169,7 +170,7 @@ public class FetchManager implements FetchDownloadProgressOwner, FetchEnqueueRes
 		return ECompleteReason.Success;
 	}
 
-	public void QueueNewDownloadDescription(DownloadDescription Description)
+	public void QueueNewDownloadDescription(DownloadDescription Description, boolean LastRequest)
 	{
 		//We have hit this code after something else has already invalidated our FetchInstance
 		if (!IsFetchInstanceValid())
@@ -177,7 +178,6 @@ public class FetchManager implements FetchDownloadProgressOwner, FetchEnqueueRes
 			Log.error("Call to QueueNewDownloadDescription after FetchInstance is invalidated! RequestID:" + Description.RequestID);
 			return;
 		}
-		
 		RequestedDownloads.put(Description.RequestID, Description);
 		ResetDownloadCompleteTracking(Description);
 
@@ -204,7 +204,7 @@ public class FetchManager implements FetchDownloadProgressOwner, FetchEnqueueRes
 		{
 			ECompleteReason CompleteReason = GetCompleteReasonForDownload(Description);
 			Log.debug("Completed RequestID:" + Description.RequestID + " with Reason " + CompleteReason + " as request was previously completed.");
-			CompleteDownload(Description, CompleteReason);
+			CompleteDownload(Description, CompleteReason, LastRequest);
 		}
 	}
 
@@ -293,7 +293,7 @@ public class FetchManager implements FetchDownloadProgressOwner, FetchEnqueueRes
 		return FetchRequest;
 	}
 	
-	private void HandleChangedDownloadDescription(DownloadDescription OldDescription, DownloadDescription NewDescription)
+	private void HandleChangedDownloadDescription(DownloadDescription OldDescription, DownloadDescription NewDescription, boolean LastRequest)
 	{
 		if (!IsFetchInstanceValid())
 		{
@@ -305,7 +305,7 @@ public class FetchManager implements FetchDownloadProgressOwner, FetchEnqueueRes
 		CopyStateToNewDescription(OldDescription, NewDescription);
 
 		//Handle the change by cancelling and recreating the fetch2 download
-		RecreateDownloadByTagFunc RecreateFunc = new RecreateDownloadByTagFunc(this, NewDescription);
+		RecreateDownloadByTagFunc RecreateFunc = new RecreateDownloadByTagFunc(this, NewDescription, LastRequest);
 		FetchInstance.getDownloadsByTag(NewDescription.RequestID, RecreateFunc);
 	}
 	
@@ -441,7 +441,7 @@ public class FetchManager implements FetchDownloadProgressOwner, FetchEnqueueRes
 				super.call(MatchingDownloads);
 				
 				//Now just have the fetch manager owner recreate the download with the supplied data
-				Owner.QueueNewDownloadDescription(RecreateDescription);
+				Owner.QueueNewDownloadDescription(RecreateDescription, LastRequest);
 			}
 			else
 			{
@@ -449,10 +449,11 @@ public class FetchManager implements FetchDownloadProgressOwner, FetchEnqueueRes
 			}
 		}
 		
-		public RecreateDownloadByTagFunc(FetchManager Owner, DownloadDescription RecreateDescription)
+		public RecreateDownloadByTagFunc(FetchManager Owner, DownloadDescription RecreateDescription, boolean LastRequest)
 		{
 			this.Owner = Owner;
 			this.RecreateDescription = RecreateDescription;
+			this.LastRequest = LastRequest;
 		}
 		
 		private boolean IsValid()
@@ -460,6 +461,7 @@ public class FetchManager implements FetchDownloadProgressOwner, FetchEnqueueRes
 			return ((null != Owner) && (null != RecreateDescription) && IsFetchInstanceValid());
 		}
 		
+		private boolean LastRequest;
 		private FetchManager Owner;
 		private	DownloadDescription RecreateDescription;
 	}
@@ -667,7 +669,7 @@ public class FetchManager implements FetchDownloadProgressOwner, FetchEnqueueRes
 			return;
 		}
 		
-		CompleteDownload(MatchedDownload, completeReason);
+		CompleteDownload(MatchedDownload, completeReason, true);
 	}
 
 	//
@@ -721,7 +723,7 @@ public class FetchManager implements FetchDownloadProgressOwner, FetchEnqueueRes
 		if (!IsFetchInstanceValid())
 		{
 			Log.debug("Call to RetryDownload with an invalid fetch instance. Skipping fetch work and completing to avoid crashing. RequestID:" + RequestID);
-			CompleteDownload(MatchingDescription, FetchRequestProgressListener.ECompleteReason.Error);
+			CompleteDownload(MatchingDescription, FetchRequestProgressListener.ECompleteReason.Error, true);
 			return;
 		}
 
@@ -730,7 +732,7 @@ public class FetchManager implements FetchDownloadProgressOwner, FetchEnqueueRes
 
 		if (IsOutOfRetries(MatchingDescription))
 		{
-			CompleteDownload(MatchingDescription, FetchRequestProgressListener.ECompleteReason.OutOfRetries);
+			CompleteDownload(MatchingDescription, FetchRequestProgressListener.ECompleteReason.OutOfRetries, true);
 			return;
 		}
 		
@@ -752,7 +754,7 @@ public class FetchManager implements FetchDownloadProgressOwner, FetchEnqueueRes
 		if (!IsFetchInstanceValid())
 		{
 			Log.debug("Call to RetryDownload_Internal with an invalid fetch instance. Skipping fetch work and completing to avoid crashing. RequestID:" + DownloadDesc.RequestID);
-			CompleteDownload(DownloadDesc, FetchRequestProgressListener.ECompleteReason.Error);
+			CompleteDownload(DownloadDesc, FetchRequestProgressListener.ECompleteReason.Error, true);
 			return;
 		}
 		
@@ -766,7 +768,7 @@ public class FetchManager implements FetchDownloadProgressOwner, FetchEnqueueRes
 		}
 		else
 		{
-			QueueNewDownloadDescription(DownloadDesc);
+			QueueNewDownloadDescription(DownloadDesc, true);
 		}
 	}
 
@@ -797,7 +799,7 @@ public class FetchManager implements FetchDownloadProgressOwner, FetchEnqueueRes
 		{
 			if (null != CachedDownloadDescription) 
 			{
-				QueueNewDownloadDescription(CachedDownloadDescription);
+				QueueNewDownloadDescription(CachedDownloadDescription, true);
 			}
 		}
 		
@@ -873,7 +875,7 @@ public class FetchManager implements FetchDownloadProgressOwner, FetchEnqueueRes
 		}
 	}
 
-	private void CompleteDownload(DownloadDescription DownloadDesc, ECompleteReason CompleteReason)
+	private void CompleteDownload(DownloadDescription DownloadDesc, ECompleteReason CompleteReason, boolean CheckForAllCompleted)
 	{
 		if (false == HasValidProgressCallback(DownloadDesc))
 		{
@@ -932,8 +934,11 @@ public class FetchManager implements FetchDownloadProgressOwner, FetchEnqueueRes
 			FailedDownloads.put(DownloadDesc.RequestID, DownloadDesc);
 		}
 
-		DownloadDesc.ProgressListener.OnDownloadComplete(DownloadDesc.RequestID, DownloadDesc.DestinationLocation, CompleteReasonToSend);		
-		CheckForAllDownloadsComplete(DownloadDesc.ProgressListener);
+		DownloadDesc.ProgressListener.OnDownloadComplete(DownloadDesc.RequestID, DownloadDesc.DestinationLocation, CompleteReasonToSend);	
+		if (CheckForAllCompleted)
+		{
+			CheckForAllDownloadsComplete(DownloadDesc.ProgressListener);
+		}
 	}
 	
 	private EDownloadCompleteReason ConvertCompleteReasonForDownload(ECompleteReason CompleteReason)
