Adding network listener for Android.
Adding usage of network listener in Background Download.

#rnx

diff --git a/Engine/Build/Android/Java/src/com/epicgames/unreal/GameActivity.java.template b/Engine/Build/Android/Java/src/com/epicgames/unreal/GameActivity.java.template
--- a/Engine/Build/Android/Java/src/com/epicgames/unreal/GameActivity.java.template
+++ b/Engine/Build/Android/Java/src/com/epicgames/unreal/GameActivity.java.template
@@ -197,6 +197,10 @@ import android.os.Build;
 
 import com.epicgames.unreal.DownloadShim;
 
+// Network connection change tracking
+import com.epicgames.unreal.network.NetworkConnectivityClient;
+import com.epicgames.unreal.network.NetworkChangedManager;
+
 // used in new virtual keyboard
 import android.view.inputmethod.InputMethodManager;
 
@@ -279,6 +283,9 @@ public class GameActivity extends $${gameActivitySuperClass}$$ implements Surfac
 
 	public static Logger Log = new Logger("UE", "GameActivity");
 
+	// For Network connection change tracking
+	private NetworkConnectivityClient.Listener NetworkListener = null;
+
 	// From cpufeatures.h AndroidCpuFamily enum
 	private static final int ANDROID_CPU_FAMILY_UNKNOWN = 0;
 	private static final int ANDROID_CPU_FAMILY_ARM = 1;
@@ -5764,7 +5771,36 @@ public class GameActivity extends $${gameActivitySuperClass}$$ implements Surfac
 		}
 		return bAirplaneMode ? CONNECTION_TYPE_AIRPLANEMODE : CONNECTION_TYPE_NONE;
 	}
-	
+
+	public void AndroidThunkJava_AddNetworkListener()
+	{
+		if (NetworkListener == null)
+		{
+			NetworkListener = new NetworkConnectivityClient.Listener() {
+				@Override
+				public void onNetworkAvailable(NetworkConnectivityClient.NetworkTransportType networkTransportType)
+				{
+					nativeNetworkChanged();
+				}
+				@Override
+				public void onNetworkLost()
+				{
+					nativeNetworkChanged();
+				}
+			};
+			NetworkChangedManager.getInstance().addListener(NetworkListener);
+		}
+	}
+
+	public void AndroidThunkJava_RemoveNetworkListener()
+	{
+		if (NetworkListener != null)
+		{
+			NetworkChangedManager.getInstance().removeListener(NetworkListener);
+			NetworkListener = null;
+		}
+	}
+
 	public boolean AndroidThunkJava_HasMetaDataKey(String key)
 	{
 		if (_bundle == null || key == null)
@@ -6956,6 +6992,8 @@ public class GameActivity extends $${gameActivitySuperClass}$$ implements Surfac
 	public native void nativeCrashContextSetFloatKey(String Key, float Value);
 	public native void nativeCrashContextSetDoubleKey(String Key, double Value);
 
+	public native void nativeNetworkChanged();
+
 	static
 	{
 		System.loadLibrary("Unreal");
diff --git a/Engine/Plugins/Online/Android/AndroidFetchBackgroundDownload/Source/AndroidFetchBackgroundDownload/AndroidFetchBackgroundDownload_UPL.xml b/Engine/Plugins/Online/Android/AndroidFetchBackgroundDownload/Source/AndroidFetchBackgroundDownload/AndroidFetchBackgroundDownload_UPL.xml
--- a/Engine/Plugins/Online/Android/AndroidFetchBackgroundDownload/Source/AndroidFetchBackgroundDownload/AndroidFetchBackgroundDownload_UPL.xml
+++ b/Engine/Plugins/Online/Android/AndroidFetchBackgroundDownload/Source/AndroidFetchBackgroundDownload/AndroidFetchBackgroundDownload_UPL.xml
@@ -47,48 +47,4 @@
 		<log text="Copying AndroidFetchBackgroundDownload files to staging"/>
 		<copyDir src="$S(PluginDir)/Java" dst="$S(BuildDir)/src/com/epicgames/unreal" />
 	</prebuildCopies>
-
-	<gameActivityImportAdditions>
-		<insert>
-			import com.epicgames.unreal.network.NetworkConnectivityClient;
-			import com.epicgames.unreal.network.NetworkChangedManager;
-		</insert>
-	</gameActivityImportAdditions>
-
-	<gameActivityClassAdditions>
-		<insert>
-			private NetworkConnectivityClient.Listener AFBD_NetworkListener = null;
-
-			public native void nativeAndroidBackgroundDownloadNetworkChanged();
-
-			public void AndroidThunkJava_AFBD_AddNetworkListener()
-			{
-				if (AFBD_NetworkListener == null)
-				{
-					AFBD_NetworkListener = new NetworkConnectivityClient.Listener() {
-						@Override
-						public void onNetworkAvailable(NetworkConnectivityClient.NetworkTransportType networkTransportType)
-						{
-							nativeAndroidBackgroundDownloadNetworkChanged();
-						}
-						@Override
-						public void onNetworkLost()
-						{
-							nativeAndroidBackgroundDownloadNetworkChanged();
-						}
-					};
-					NetworkChangedManager.getInstance().addListener(AFBD_NetworkListener);
-				}
-			}
-
-			public void AndroidThunkJava_AFBD_RemoveNetworkListener()
-			{
-				if (AFBD_NetworkListener != null)
-				{
-					NetworkChangedManager.getInstance().removeListener(AFBD_NetworkListener);
-					AFBD_NetworkListener = null;
-				}
-			}
-		</insert>
-	</gameActivityClassAdditions>
 </root>
diff --git a/Engine/Plugins/Online/Android/AndroidFetchBackgroundDownload/Source/AndroidFetchBackgroundDownload/Private/AndroidPlatformBackgroundHttpManager.cpp b/Engine/Plugins/Online/Android/AndroidFetchBackgroundDownload/Source/AndroidFetchBackgroundDownload/Private/AndroidPlatformBackgroundHttpManager.cpp
--- a/Engine/Plugins/Online/Android/AndroidFetchBackgroundDownload/Source/AndroidFetchBackgroundDownload/Private/AndroidPlatformBackgroundHttpManager.cpp
+++ b/Engine/Plugins/Online/Android/AndroidFetchBackgroundDownload/Source/AndroidFetchBackgroundDownload/Private/AndroidPlatformBackgroundHttpManager.cpp
@@ -51,7 +51,6 @@ namespace FAndroidBackgroundDownloadDelegates
 	DECLARE_DELEGATE_FourParams(FAndroidBackgroundDownload_OnComplete, jobject /*UnderlyingWorker*/, FString /*RequestID*/, FString /*CompleteLocation*/, bool /*bWasSuccess*/);
 	DECLARE_DELEGATE_TwoParams(FAndroidBackgroundDownload_OnAllComplete, jobject /*UnderlyingWorker*/, bool /*bDidAllRequestsSucceed*/);
 	DECLARE_DELEGATE_TwoParams(FAndroidBackgroundDownload_OnTickWorkerThread, JNIEnv*, jobject /*UnderlyingWorker*/);
-	DECLARE_DELEGATE_TwoParams(FAndroidBackgroundDownload_OnNetworkChanged, JNIEnv*, bool /*bNetworkIsConnected*/);
 
 	//Delegates called by JNI functions to bubble up underlying java work to the manager
 	FAndroidBackgroundDownload_OnWorkerStart AndroidBackgroundDownload_OnWorkerStart;
@@ -60,7 +59,6 @@ namespace FAndroidBackgroundDownloadDelegates
 	FAndroidBackgroundDownload_OnComplete AndroidBackgroundDownload_OnComplete;
 	FAndroidBackgroundDownload_OnAllComplete AndroidBackgroundDownload_OnAllComplete;
 	FAndroidBackgroundDownload_OnTickWorkerThread AndroidBackgroundDownload_OnTickWorkerThread;
-	FAndroidBackgroundDownload_OnNetworkChanged AndroidBackgroundDownload_OnNetworkChanged;
 };
 
 const FString FAndroidPlatformBackgroundHttpManager::BackgroundHTTPWorkID = TEXT("BackgroundHttpDownload");
@@ -120,13 +118,10 @@ FAndroidPlatformBackgroundHttpManager::FAndroidPlatformBackgroundHttpManager()
 	, bWorkerHadTerminalFinish(false)
 	, AndroidBackgroundHTTPManagerDefaultLocalizedText()
 {
-	if (JNIEnv* Env = FAndroidApplication::GetJavaEnv())
+	if (!OnNetworkConnectionChangedDelegateHandle.IsValid())
 	{
-		static jmethodID AddNetworkListenerFunc = FJavaWrapper::FindMethod(Env, FJavaWrapper::GameActivityClassID, "AndroidThunkJava_AFBD_AddNetworkListener", "()V", false);
-		if (AddNetworkListenerFunc != nullptr)
-		{
-			FJavaWrapper::CallVoidMethod(Env, FJavaWrapper::GameActivityThis, AddNetworkListenerFunc);
-		}
+		OnNetworkConnectionChangedDelegateHandle = FPlatformMisc::AddNetworkListener(
+			FCoreDelegates::FOnNetworkConnectionChanged::FDelegate::CreateRaw(this, &FAndroidPlatformBackgroundHttpManager::OnNetworkConnectionChanged));
 	}
 
 	bNetworkConnected = GetNetworkConnected();
@@ -172,7 +167,6 @@ void FAndroidPlatformBackgroundHttpManager::Initialize()
 	FAndroidBackgroundDownloadDelegates::AndroidBackgroundDownload_OnComplete.BindSP(this, &FAndroidPlatformBackgroundHttpManager::Java_OnDownloadComplete);
 	FAndroidBackgroundDownloadDelegates::AndroidBackgroundDownload_OnAllComplete.BindSP(this, &FAndroidPlatformBackgroundHttpManager::Java_OnAllDownloadsComplete);
 	FAndroidBackgroundDownloadDelegates::AndroidBackgroundDownload_OnTickWorkerThread.BindSP(this, &FAndroidPlatformBackgroundHttpManager::Java_OnTick);
-	FAndroidBackgroundDownloadDelegates::AndroidBackgroundDownload_OnNetworkChanged.BindSP(this, &FAndroidPlatformBackgroundHttpManager::Java_OnNetworkChanged);
 
 	AndroidBackgroundHTTPManagerDefaultLocalizedText.InitFromIniSettings("AndroidFetchBackgroundDownload");
 	FBackgroundHttpManagerImpl::Initialize();
@@ -180,6 +174,9 @@ void FAndroidPlatformBackgroundHttpManager::Initialize()
 
 void FAndroidPlatformBackgroundHttpManager::Shutdown()
 {
+	FPlatformMisc::RemoveNetworkListener(OnNetworkConnectionChangedDelegateHandle);
+	OnNetworkConnectionChangedDelegateHandle.Reset();
+
 	FBackgroundHttpManagerImpl::Shutdown();
 }
 
@@ -895,16 +892,6 @@ void FAndroidPlatformBackgroundHttpManager::Java_OnTick(JNIEnv* Env, jobject Und
 		}
 	}
 }
-void FAndroidPlatformBackgroundHttpManager::Java_OnNetworkChanged(JNIEnv* Env, bool bNetworkIsConnected)
-{
-	if (bNetworkIsConnected)
-	{
-		UE_LOG(LogBackgroundHttpManager, Display, TEXT("Flagging worker for requeue due to network connection re-establishing."));
-		//On reconnecting to the network flag our worker to force requeue to make sure our work is continued
-		FPlatformAtomics::InterlockedExchange(&bShouldForceWorkerRequeue, true);
-	}
-	
-}
 
 void FAndroidPlatformBackgroundHttpManager::ForceCompleteAllUnderlyingJavaActiveRequestsWithError()
 {
@@ -1136,6 +1123,20 @@ void FAndroidPlatformBackgroundHttpManager::FAndroidBackgroundHTTPManagerDefault
 	}
 }
 
+void FAndroidPlatformBackgroundHttpManager::OnNetworkConnectionChanged(ENetworkConnectionType ConnectionType)
+{
+	bool bNewNetworkConnected = GetNetworkConnected();
+
+	if (bNewNetworkConnected != bNetworkConnected)
+	{
+		bNetworkConnected = bNewNetworkConnected;
+
+		UE_LOG(LogBackgroundHttpManager, Display, TEXT("Flagging worker for requeue due to network connection re-establishing."));
+		//On reconnecting to the network flag our worker to force requeue to make sure our work is continued
+		FPlatformAtomics::InterlockedExchange(&bShouldForceWorkerRequeue, true);
+	}
+}
+
 //
 //JNI methods coming from UEDownloadWorker
 //
@@ -1193,20 +1194,5 @@ JNI_METHOD void Java_com_epicgames_unreal_download_UEDownloadWorker_nativeAndroi
 	FAndroidBackgroundDownloadDelegates::AndroidBackgroundDownload_OnTickWorkerThread.ExecuteIfBound(jenv, thiz);
 }
 
-//
-//JNI methods coming from GameActivity
-//
-
-JNI_METHOD void Java_com_epicgames_unreal_GameActivity_nativeAndroidBackgroundDownloadNetworkChanged(JNIEnv* jenv, jobject thiz)
-{
-	bool bNewNetworkConnected = GetNetworkConnected();
-
-	if (bNewNetworkConnected != bNetworkConnected)
-	{
-		bNetworkConnected = bNewNetworkConnected;
-		FAndroidBackgroundDownloadDelegates::AndroidBackgroundDownload_OnNetworkChanged.ExecuteIfBound(jenv, bNetworkConnected);
-	}
-}
-
 
 #undef LOCTEXT_NAMESPACE
\ No newline at end of file
diff --git a/Engine/Plugins/Online/Android/AndroidFetchBackgroundDownload/Source/AndroidFetchBackgroundDownload/Public/AndroidPlatformBackgroundHttpManager.h b/Engine/Plugins/Online/Android/AndroidFetchBackgroundDownload/Source/AndroidFetchBackgroundDownload/Public/AndroidPlatformBackgroundHttpManager.h
--- a/Engine/Plugins/Online/Android/AndroidFetchBackgroundDownload/Source/AndroidFetchBackgroundDownload/Public/AndroidPlatformBackgroundHttpManager.h
+++ b/Engine/Plugins/Online/Android/AndroidFetchBackgroundDownload/Source/AndroidFetchBackgroundDownload/Public/AndroidPlatformBackgroundHttpManager.h
@@ -76,7 +76,6 @@ protected:
 	void Java_OnDownloadComplete(jobject UnderlyingWorker, FString RequestID, FString CompleteLocation, bool bWasSuccess);
 	void Java_OnAllDownloadsComplete(jobject UnderlyingWorker, bool bDidAllRequestsSucceed);
 	void Java_OnTick(JNIEnv* Env, jobject UnderlyingWorker);
-	void Java_OnNetworkChanged(JNIEnv* Env, bool bNetworkIsConnected);
 
 	//Helper function that completes any un-completed requests in the underlying java tracking layer so that they will
 	//complete on the next Tick. Mostly used to help when the underlying java error hits a failure. Ensures when any
@@ -125,6 +124,9 @@ protected:
 	//Rechecks any _GT lists to try and move them to _Java lists if its safe to do so
 	void HandleRequestsWaitingOnJavaThreadsafety();
 
+	void OnNetworkConnectionChanged(ENetworkConnectionType ConnectionType);
+	FDelegateHandle OnNetworkConnectionChangedDelegateHandle;
+
 protected:
 	//Used to determine the state of AndroidBackgroundHTTP in our configrules
 	enum class EAndroidBackgroundDownloadConfigRulesSetting : uint8
diff --git a/Engine/Source/Runtime/Core/Private/Android/AndroidPlatformMisc.cpp b/Engine/Source/Runtime/Core/Private/Android/AndroidPlatformMisc.cpp
--- a/Engine/Source/Runtime/Core/Private/Android/AndroidPlatformMisc.cpp
+++ b/Engine/Source/Runtime/Core/Private/Android/AndroidPlatformMisc.cpp
@@ -6,6 +6,7 @@
 #include "Misc/FileHelper.h"
 #include "Misc/App.h"
 #include "Misc/ConfigCacheIni.h"
+#include "Misc/CoreDelegates.h"
 #include "Misc/FeedbackContext.h"
 #include "Misc/OutputDeviceRedirector.h"
 #include "HAL/IConsoleManager.h"
@@ -576,8 +577,12 @@ void FAndroidMisc::PlatformInit()
 	InitializeJavaEventReceivers();
 	AndroidOnBackgroundBinding = FCoreDelegates::ApplicationWillEnterBackgroundDelegate.AddStatic(EnableJavaEventReceivers, false);
 	AndroidOnForegroundBinding = FCoreDelegates::ApplicationHasEnteredForegroundDelegate.AddStatic(EnableJavaEventReceivers, true);
+
+	extern void AndroidThunkJava_AddNetworkListener();
+	AndroidThunkJava_AddNetworkListener();
 #endif
 
+
 	InitCpuThermalSensor();
 
 	UE_LOG(LogInit, Log, TEXT(" - This binary is optimized with LTO: %s, PGO: %s, instrumented for PGO data collection: %s"),
@@ -2831,6 +2836,17 @@ bool FAndroidMisc::HasActiveWiFiConnection()
 }
 #endif
 
+JNI_METHOD void Java_com_epicgames_unreal_GameActivity_nativeNetworkChanged(JNIEnv* jenv, jobject thiz)
+{
+	if (FTaskGraphInterface::IsRunning())
+	{
+		FFunctionGraphTask::CreateAndDispatchWhenReady([]()
+		{
+			FCoreDelegates::OnNetworkConnectionChanged.Broadcast(FAndroidMisc::GetNetworkConnectionType());
+		}, TStatId(), NULL, ENamedThreads::GameThread);
+	}
+}
+
 static FAndroidMisc::ReInitWindowCallbackType OnReInitWindowCallback;
 
 FAndroidMisc::ReInitWindowCallbackType FAndroidMisc::GetOnReInitWindowCallback()
@@ -3163,3 +3179,32 @@ void FAndroidMisc::ShowConsoleWindow()
 	AndroidThunkCpp_ShowConsoleWindow();
 #endif // !UE_BUILD_SHIPPING && USE_ANDROID_JNI
 }
+
+FDelegateHandle FAndroidMisc::AddNetworkListener(FCoreDelegates::FOnNetworkConnectionChanged::FDelegate&& InNewDelegate)
+{
+	if (!FCoreDelegates::OnNetworkConnectionChanged.IsBound())
+	{
+#if USE_ANDROID_JNI
+		extern void AndroidThunkJava_AddNetworkListener();
+		AndroidThunkJava_AddNetworkListener();
+#endif
+	}
+
+	return FCoreDelegates::OnNetworkConnectionChanged.Add(MoveTemp(InNewDelegate));
+}
+
+bool FAndroidMisc::RemoveNetworkListener(FDelegateHandle Handle)
+{
+	bool bSuccess = FCoreDelegates::OnNetworkConnectionChanged.Remove(Handle);
+
+	if (!FCoreDelegates::OnNetworkConnectionChanged.IsBound())
+	{
+#if USE_ANDROID_JNI
+		extern void AndroidThunkJava_RemoveNetworkListener();
+		AndroidThunkJava_RemoveNetworkListener();
+#endif
+	}
+
+	return bSuccess;
+}
+
diff --git a/Engine/Source/Runtime/Core/Private/GenericPlatform/GenericPlatformMisc.cpp b/Engine/Source/Runtime/Core/Private/GenericPlatform/GenericPlatformMisc.cpp
--- a/Engine/Source/Runtime/Core/Private/GenericPlatform/GenericPlatformMisc.cpp
+++ b/Engine/Source/Runtime/Core/Private/GenericPlatform/GenericPlatformMisc.cpp
@@ -3,6 +3,7 @@
 #include "GenericPlatform/GenericPlatformMisc.h"
 #include "HAL/IConsoleManager.h"
 #include "Misc/AssertionMacros.h"
+#include "Misc/CoreDelegates.h"
 #include "HAL/PlatformFileManager.h"
 #include "HAL/CriticalSection.h"
 #include "Misc/ScopeRWLock.h"
@@ -2025,4 +2026,18 @@ void FGenericPlatformMisc::ShowConsoleWindow()
 #if !UE_BUILD_SHIPPING
 	UE_LOG(LogGenericPlatformMisc, Log, TEXT("Show console is not supported or implemented in current platform"));
 #endif
-}
\ No newline at end of file
+}
+
+FDelegateHandle FGenericPlatformMisc::AddNetworkListener(FCoreDelegates::FOnNetworkConnectionChanged::FDelegate&& InNewDelegate)
+{
+	UE_LOG(LogGenericPlatformMisc, Warning, TEXT("FGenericPlatformMisc::AddNetworkListener not implemented for this platform"));
+
+	return FDelegateHandle();
+}
+
+bool FGenericPlatformMisc::RemoveNetworkListener(FDelegateHandle Handle)
+{
+	UE_LOG(LogGenericPlatformMisc, Warning, TEXT("FGenericPlatformMisc::RemoveNetworkListener not implemented for this platform"));
+
+	return false;
+}
diff --git a/Engine/Source/Runtime/Core/Private/Misc/CoreDelegates.cpp b/Engine/Source/Runtime/Core/Private/Misc/CoreDelegates.cpp
--- a/Engine/Source/Runtime/Core/Private/Misc/CoreDelegates.cpp
+++ b/Engine/Source/Runtime/Core/Private/Misc/CoreDelegates.cpp
@@ -112,6 +112,7 @@ FSimpleMulticastDelegate FCoreDelegates::OnSamplingInput;
 FSimpleMulticastDelegate FCoreDelegates::OnEndFrame;
 FSimpleMulticastDelegate FCoreDelegates::OnBeginFrameRT;
 FSimpleMulticastDelegate FCoreDelegates::OnEndFrameRT;
+FCoreDelegates::FOnNetworkConnectionChanged FCoreDelegates::OnNetworkConnectionChanged;
 TDelegate<EAppReturnType::Type(EAppMsgType::Type, const FText&, const FText&)> FCoreDelegates::ModalErrorMessage;
 TMulticastDelegate<void(const FString&, const FString&)> FCoreDelegates::OnInviteAccepted;
 TMulticastDelegate<void(class UWorld*, FIntVector, FIntVector)> FCoreDelegates::PreWorldOriginOffset;
diff --git a/Engine/Source/Runtime/Core/Public/Android/AndroidPlatformMisc.h b/Engine/Source/Runtime/Core/Public/Android/AndroidPlatformMisc.h
--- a/Engine/Source/Runtime/Core/Public/Android/AndroidPlatformMisc.h
+++ b/Engine/Source/Runtime/Core/Public/Android/AndroidPlatformMisc.h
@@ -315,6 +315,9 @@ public:
 
 	static void ShowConsoleWindow();
 
+	static FDelegateHandle AddNetworkListener(FOnNetworkConnectionChangedDelegate&& InNewDelegate);
+	static bool RemoveNetworkListener(FDelegateHandle Handle);
+
 private:
 	static const ANSICHAR* CodeToString(int Signal, int si_code);
 	static EDeviceScreenOrientation DeviceOrientation;
diff --git a/Engine/Source/Runtime/Core/Public/GenericPlatform/GenericPlatformMisc.h b/Engine/Source/Runtime/Core/Public/GenericPlatform/GenericPlatformMisc.h
--- a/Engine/Source/Runtime/Core/Public/GenericPlatform/GenericPlatformMisc.h
+++ b/Engine/Source/Runtime/Core/Public/GenericPlatform/GenericPlatformMisc.h
@@ -17,6 +17,7 @@
 #endif
 
 class Error;
+class FDelegateHandle;
 class FOutputDevice;
 class FString;
 class FText;
@@ -26,6 +27,7 @@ class IPlatformChunkInstall;
 class IPlatformCompression;
 class IPlatformHostCommunication;
 struct FCustomChunk;
+struct FDefaultDelegateUserPolicy;
 struct FGenericCrashContext;
 struct FGenericMemoryWarningContext;
 struct FGuid;
@@ -35,6 +37,9 @@ enum class ECustomChunkType : uint8;
 template <typename FuncType>
 class TFunction;
 
+template <typename FuncType, typename UserPolicy>
+class TDelegate;
+
 #if UE_BUILD_SHIPPING && !WITH_EDITOR
 #define UE_DEBUG_BREAK() ((void)0)
 #else
@@ -371,6 +376,8 @@ enum class ENetworkConnectionType : uint8
 	Ethernet,
 };
 
+typedef TDelegate<void(ENetworkConnectionType ConnectionType), FDefaultDelegateUserPolicy> FOnNetworkConnectionChangedDelegate;
+
 enum class EProcessDiagnosticFlags : uint32
 {
 	None            = 0,
@@ -540,6 +547,10 @@ struct CORE_API FGenericPlatformMisc
 	static void PlatformPreInit();
 	static void PlatformInit() { }
 
+
+	static FDelegateHandle AddNetworkListener(FOnNetworkConnectionChangedDelegate&& InNewDelegate);
+	static bool RemoveNetworkListener(FDelegateHandle Handle);
+	
 	/**
 	* Called to dismiss splash screen
 	*/
diff --git a/Engine/Source/Runtime/Core/Public/Misc/CoreDelegates.h b/Engine/Source/Runtime/Core/Public/Misc/CoreDelegates.h
--- a/Engine/Source/Runtime/Core/Public/Misc/CoreDelegates.h
+++ b/Engine/Source/Runtime/Core/Public/Misc/CoreDelegates.h
@@ -254,6 +254,9 @@ public:
 	/** Called when the OS goes into low power mode */
 	static TMulticastDelegate<void(bool)> OnLowPowerMode;
 
+	DECLARE_MULTICAST_DELEGATE_OneParam(FOnNetworkConnectionChanged, ENetworkConnectionType /*ConnectionType*/);
+	static FOnNetworkConnectionChanged OnNetworkConnectionChanged;
+
 	static TTSMulticastDelegate<void(const TCHAR* IniFilename, int32& ResponderCount)>& TSCountPreLoadConfigFileRespondersDelegate();
 	static TTSMulticastDelegate<void(const TCHAR* IniFilename, FString& LoadedContents)>& TSPreLoadConfigFileDelegate();
 	static TTSMulticastDelegate<void(const TCHAR* IniFilename, const FString& ContentsToSave, int32& SavedCount)>& TSPreSaveConfigFileDelegate();
diff --git a/Engine/Source/Runtime/Launch/Private/Android/AndroidJNI.cpp b/Engine/Source/Runtime/Launch/Private/Android/AndroidJNI.cpp
--- a/Engine/Source/Runtime/Launch/Private/Android/AndroidJNI.cpp
+++ b/Engine/Source/Runtime/Launch/Private/Android/AndroidJNI.cpp
@@ -112,6 +112,8 @@ void FJavaWrapper::FindClassesAndMethods(JNIEnv* Env)
 	AndroidThunkJava_LocalNotificationGetLaunchNotification = FindMethod(Env, GameActivityClassID, "AndroidThunkJava_LocalNotificationGetLaunchNotification", "()Lcom/epicgames/unreal/GameActivity$LaunchNotification;", bIsOptional);
 	AndroidThunkJava_LocalNotificationDestroyIfExists = FindMethod(Env, GameActivityClassID, "AndroidThunkJava_LocalNotificationDestroyIfExists", "(I)Z", bIsOptional);
 	AndroidThunkJava_GetNetworkConnectionType = FindMethod(Env, GameActivityClassID, "AndroidThunkJava_GetNetworkConnectionType", "()I", bIsOptional);
+	AndroidThunkJava_AddNetworkListener = FindMethod(Env, GameActivityClassID, "AndroidThunkJava_AddNetworkListener", "()V", bIsOptional);
+	AndroidThunkJava_RemoveNetworkListener = FindMethod(Env, GameActivityClassID, "AndroidThunkJava_AddNetworkListener", "()V", bIsOptional);
 	AndroidThunkJava_GetAndroidId = FindMethod(Env, GameActivityClassID, "AndroidThunkJava_GetAndroidId", "()Ljava/lang/String;", bIsOptional);
 	AndroidThunkJava_ShareURL = FindMethod(Env, GameActivityClassID, "AndroidThunkJava_ShareURL", "(Ljava/lang/String;Ljava/lang/String;Ljava/lang/String;II)V", bIsOptional);
 	AndroidThunkJava_IsPackageInstalled = FindMethod(Env, GameActivityClassID, "AndroidThunkJava_IsPackageInstalled", "(Ljava/lang/String;)Z", bIsOptional);
@@ -469,6 +471,9 @@ jmethodID FJavaWrapper::AndroidThunkJava_GetSupportedNativeDisplayRefreshRates;
 jmethodID FJavaWrapper::AndroidThunkJava_GetNativeDisplayRefreshRate;
 jmethodID FJavaWrapper::AndroidThunkJava_SetNativeDisplayRefreshRate;
 
+jmethodID FJavaWrapper::AndroidThunkJava_AddNetworkListener;
+jmethodID FJavaWrapper::AndroidThunkJava_RemoveNetworkListener;
+
 jmethodID FJavaWrapper::AndroidThunkJava_EnableMotion;
 
 jclass FJavaWrapper::LaunchNotificationClass;
@@ -1803,6 +1808,22 @@ TArray<int32> AndroidThunkCpp_GetSupportedNativeDisplayRefreshRates()
 	return Result;
 }
 
+void AndroidThunkJava_AddNetworkListener()
+{
+	if (JNIEnv* Env = FAndroidApplication::GetJavaEnv())
+	{
+		FJavaWrapper::CallVoidMethod(Env, FJavaWrapper::GameActivityThis, FJavaWrapper::AndroidThunkJava_AddNetworkListener);
+	}
+}
+
+void AndroidThunkJava_RemoveNetworkListener()
+{
+	if (JNIEnv* Env = FAndroidApplication::GetJavaEnv())
+	{
+		FJavaWrapper::CallVoidMethod(Env, FJavaWrapper::GameActivityThis, FJavaWrapper::AndroidThunkJava_RemoveNetworkListener);
+	}
+}
+
 bool AndroidThunkCpp_SetNativeDisplayRefreshRate(int32 RefreshRate)
 {
 	bool Result = false;
diff --git a/Engine/Source/Runtime/Launch/Public/Android/AndroidJNI.h b/Engine/Source/Runtime/Launch/Public/Android/AndroidJNI.h
--- a/Engine/Source/Runtime/Launch/Public/Android/AndroidJNI.h
+++ b/Engine/Source/Runtime/Launch/Public/Android/AndroidJNI.h
@@ -135,6 +135,10 @@ public:
 	// Motion controls
 	static jmethodID AndroidThunkJava_EnableMotion;
 
+	// Network Connection Listener
+	static jmethodID AndroidThunkJava_AddNetworkListener;
+	static jmethodID AndroidThunkJava_RemoveNetworkListener;
+
 	// member fields for getting the launch notification
 	static jclass LaunchNotificationClass;
 	static jfieldID LaunchNotificationUsed;
