Android Background HTTP fixes:
- Now unregisters and cleans up FetchListener appropriately to fix bugs from receiving underlying FetchListener callbacks on old invalid FetchListeners.
- Now capture logging from underlying FetchInstance through new FUEFetchLogger that overrides the underlying FetchLogger
- Fixing a bug around FetchInstances not being properly cleaned up and closed potentially if the StopWork took too long to process and the OS killed that thread
#tests Android Dev Client

diff --git a/Engine/Plugins/Online/Android/AndroidFetchBackgroundDownload/Source/AndroidFetchBackgroundDownload/Java/download/fetch/FetchManager.java b/Engine/Plugins/Online/Android/AndroidFetchBackgroundDownload/Source/AndroidFetchBackgroundDownload/Java/download/fetch/FetchManager.java
--- a/Engine/Plugins/Online/Android/AndroidFetchBackgroundDownload/Source/AndroidFetchBackgroundDownload/Java/download/fetch/FetchManager.java
+++ b/Engine/Plugins/Online/Android/AndroidFetchBackgroundDownload/Source/AndroidFetchBackgroundDownload/Java/download/fetch/FetchManager.java
@@ -53,6 +53,16 @@ public class FetchManager implements FetchDownloadProgressOwner, FetchEnqueueRes
 	{
 		Log.debug("StopWork called");
 
+		//Immediately remove FetchListener to prevent us from getting late callbacks after work has stopped
+		if (FetchListener != null)
+		{
+			if (IsFetchInstanceValid())
+			{
+				FetchInstance.removeListener(FetchListener);
+			}
+			FetchListener = null;
+		}
+
 		//Synchronized with initFetch to make sure we aren't creating and closing our fetch instance in a weird race condition
 		synchronized(this)
 		{
@@ -64,6 +74,7 @@ public class FetchManager implements FetchDownloadProgressOwner, FetchEnqueueRes
 
 				//close our FetchInstance so it stops all work until we recreate it in initfetch later
 				FetchInstance.close();
+				FetchInstance = null;
 			}
 		}
 	}
@@ -370,60 +381,67 @@ public class FetchManager implements FetchDownloadProgressOwner, FetchEnqueueRes
 	{
 		Log.debug("InitFetch called");
 
+		//Remove the old FetchListener if it was setup immediately so we don't
+		//get any invalid callbacks while inside of the synchronized re-create of the FetchInstance
+		if (FetchListener != null)
+		{
+			if (IsFetchInstanceValid())
+			{
+				FetchInstance.removeListener(FetchListener);
+			}
+			FetchListener = null;
+		}
+
 		//Synchronized with StopWork to make sure we aren't creating and closing our fetch instance in a weird race condition
 		synchronized(this)
 		{
-			//Make sure any existing FetchInstance is in a correct state (either null and ready to be created, or open and unfrozen ready to do work
+			//Make sure any existing FetchInstance is in a correct state (null and ready to be created, with any in progress work deleted)
 			if (FetchInstance != null)
 			{
-				//If we previously closed our FetchInstance just remove it and recreate it bellow
-				if (FetchInstance.isClosed())
+				if (!FetchInstance.isClosed())
 				{
-					FetchInstance = null;
+					DeleteAllInProgressRequests();
+					FetchInstance.close();
 				}
-				//If our FetchInstance exists and isn't closed, its very likely frozen and needs to be unfrozen
 				else
 				{
-					Log.debug("InitFetch has existing non-closed FetchInstance. Unfreezing and deleting all in progress requests");
-
-					//If we are just unfreezing existing Fetch work, lets delete everything that wasn't finished by the previous work before unfreezing
-					//This prevents errors in resuming where we no longer want a particular download or we fail to resume the work
-					DeleteAllInProgressRequests();
-
-					//Now unfreeze (and hopefully have nothing really running)
-					FetchInstance.unfreeze();
+					Log.debug("InitFetch has existing non-closed FetchInstance!");
 				}
+
+				FetchInstance = null;
 			}
 
-			if (FetchInstance == null)
+			if (FetchInstance != null)
+			{
+				Log.error("Unexpected already configured FetchInstance after destroying FetchInstance in InitFetch!");
+			}
+			else
 			{
-				//InstanceLogger = new FetchLogger(true, "FManager");
-
 				//TODO TRoss: Pull these values from the worker's getInputData
 				FetchInstance = Fetch.Impl.getInstance(new FetchConfiguration.Builder(context)
 					.setNamespace(context.getPackageName())
-					//.enableLogging(true)
-					//.setLogger(InstanceLogger)
+					.enableLogging(true)
+					.setLogger(FetchLog)
 					.enableRetryOnNetworkGain(true)
 					.setProgressReportingInterval(200)
 					.build());
-
-				//if we are creating our FetchInstance, make sure our FetchListener is also recreated and attached
-				FetchListener = null;
 			}
-					
+			
 			if (!IsFetchInstanceValid())
 			{
 				Log.error("Unexpected invalid FetchInstance after completing InitFetch!");
 			}
 			else
-			{
-				//Add our FetchListener
-				if (null == FetchListener)
+			{	
+				//sanity check as this should have been nulled before the synchronized block
+				if (FetchListener != null)
 				{
-					FetchListener = new FetchRequestProgressListener(this);
-					FetchInstance.addListener(FetchListener);
+					Log.error("Unexpected FetchListener still active during InitFetch creation!");
+					FetchListener = null;
 				}
+
+				FetchListener = new FetchRequestProgressListener(this);
+				FetchInstance.addListener(FetchListener);
 			}
 		}
 	}
@@ -610,6 +628,7 @@ public class FetchManager implements FetchDownloadProgressOwner, FetchEnqueueRes
 	public void OnDownloadChangePauseState(@NonNull Download download, boolean bIsPaused)
 	{
 		String DownloadRequestID = GetRequestID(download);
+		Log.verbose("OnDownloadChangePauseState: " + DownloadRequestID);
 
 		DownloadDescription MatchedDownload = RequestedDownloads.get(DownloadRequestID);
 		if (null == MatchedDownload)
@@ -705,7 +724,7 @@ public class FetchManager implements FetchDownloadProgressOwner, FetchEnqueueRes
 		MatchingDescription.ProgressListener.OnDownloadEnqueued(RequestID, false);
 	}
 	
-	//if UEDOwnloadableWorker is calling into it, the only thing it should need is the RequestID
+	//if UEDownloadableWorker is calling into it, the only thing it should need is the RequestID
 	public void RetryDownload(String RequestID)
 	{
 		RetryDownload(RequestID, null);
@@ -991,6 +1010,7 @@ public class FetchManager implements FetchDownloadProgressOwner, FetchEnqueueRes
 	{
 		return ((null != DownloadDesc) && (null != DownloadDesc.ProgressListener));
 	}
+
 	//Used to determine if the reason a download completed was because of action we took. IE: Cancelling, deleting, removing, etc that happened through calls to the FetchManager.
 	//Or if it was a change of status by hitting an error, succeeding, etc.
 	private boolean IsCompleteReasonIntentional(FetchRequestProgressListener.ECompleteReason CompleteReason)
@@ -1007,10 +1027,47 @@ public class FetchManager implements FetchDownloadProgressOwner, FetchEnqueueRes
 				return false;
 		}
 	}
-	//private boolean ShouldRetryRequest(@NonNull Request 
 	
+	//Wrapper class for our FetchInstance logging that passes the values through UE's Logger Class
+	//FetchLogger only has debug and error logging currently implemented so that is all we pass through.
+	private class FUEFetchLogger extends FetchLogger
+	{
+		public FUEFetchLogger()
+		{
+			//Disable logging because we only want this overridden logging behavior
+			//and not the base logging that happens with FetchLogger
+			super(false, "UEFetchLoggerWrapper");
+		}
+
+		private Logger UELogger = new Logger("UE", "FetchAPI");
+
+		@Override 
+		public void d(String message)
+		{
+			UELogger.debug(message);
+		}
+
+		@Override
+		public void d(String message, Throwable throwable)
+		{
+			//Since we have a throwable, send this as an error with throwable as our UE Logger class doesn't have a
+			//debug with throwable
+			UELogger.error(message, throwable);
+		}
+
+		@Override
+		public void e(String message)
+		{
+			UELogger.error(message);
+		}
+
+		@Override
+		public void e(String message, Throwable throwable)
+		{
+			UELogger.error(message, throwable);
+		}
+	}
 	private volatile Fetch FetchInstance = null;
-	private FetchLogger InstanceLogger = null;
 
 	private FetchRequestProgressListener FetchListener = null;
 	
@@ -1021,4 +1078,5 @@ public class FetchManager implements FetchDownloadProgressOwner, FetchEnqueueRes
 	public String TempFileExtension = ".fetchtemp";
 	
 	public Logger Log = new Logger("UE", "FetchManager");
+	public FUEFetchLogger FetchLog = new FUEFetchLogger();
 }
\ No newline at end of file
