Background downloading Cancel implementation
- Implement CancelRequest for Background downloading
- Canceling cosmetic streaming while downloading now stops the download, and deletes the files correctly.

diff --git a/Engine/Plugins/Online/Android/AndroidFetchBackgroundDownload/Source/AndroidFetchBackgroundDownload/Java/download/fetch/FetchManager.java b/Engine/Plugins/Online/Android/AndroidFetchBackgroundDownload/Source/AndroidFetchBackgroundDownload/Java/download/fetch/FetchManager.java
--- a/Engine/Plugins/Online/Android/AndroidFetchBackgroundDownload/Source/AndroidFetchBackgroundDownload/Java/download/fetch/FetchManager.java
+++ b/Engine/Plugins/Online/Android/AndroidFetchBackgroundDownload/Source/AndroidFetchBackgroundDownload/Java/download/fetch/FetchManager.java
@@ -282,6 +282,7 @@ public class FetchManager implements FetchDownloadProgressOwner, FetchEnqueueRes
 
 			FetchInstance.delete(MatchedDesc.CachedFetchID);
 		}
+		CheckForAllDownloadsComplete(MatchedDesc.ProgressListener);
 	}
 	
 	private Request BuildFetchRequest(DownloadDescription Description)
diff --git a/Engine/Plugins/Online/Android/AndroidFetchBackgroundDownload/Source/AndroidFetchBackgroundDownload/Private/AndroidPlatformBackgroundHttpRequest.cpp b/Engine/Plugins/Online/Android/AndroidFetchBackgroundDownload/Source/AndroidFetchBackgroundDownload/Private/AndroidPlatformBackgroundHttpRequest.cpp
--- a/Engine/Plugins/Online/Android/AndroidFetchBackgroundDownload/Source/AndroidFetchBackgroundDownload/Private/AndroidPlatformBackgroundHttpRequest.cpp
+++ b/Engine/Plugins/Online/Android/AndroidFetchBackgroundDownload/Source/AndroidFetchBackgroundDownload/Private/AndroidPlatformBackgroundHttpRequest.cpp
@@ -29,6 +29,15 @@ FAndroidPlatformBackgroundHttpRequest::FAndroidPlatformBackgroundHttpRequest()
 {
 }
 
+void FAndroidPlatformBackgroundHttpRequest::CancelRequest()
+{
+	FAndroidPlatformBackgroundHttpManagerPtr AndroidBGManagerPtr = StaticCastSharedPtr<FAndroidPlatformBackgroundHttpManager>(FBackgroundHttpModule::Get().GetBackgroundHttpManager());
+	if (ensureAlwaysMsgf(AndroidBGManagerPtr.IsValid(), TEXT("Found a BackgroundHttpManager that wasn't the AndroidBackgroundHttpManager! This is not supported or expected!")))
+	{
+		AndroidBGManagerPtr->CancelRequest(SharedThis(this));
+	}
+}
+
 void FAndroidPlatformBackgroundHttpRequest::PauseRequest()
 {
 	FAndroidPlatformBackgroundHttpManagerPtr AndroidBGManagerPtr = StaticCastSharedPtr<FAndroidPlatformBackgroundHttpManager>(FBackgroundHttpModule::Get().GetBackgroundHttpManager());
diff --git a/Engine/Plugins/Online/Android/AndroidFetchBackgroundDownload/Source/AndroidFetchBackgroundDownload/Public/AndroidPlatformBackgroundHttpRequest.h b/Engine/Plugins/Online/Android/AndroidFetchBackgroundDownload/Source/AndroidFetchBackgroundDownload/Public/AndroidPlatformBackgroundHttpRequest.h
--- a/Engine/Plugins/Online/Android/AndroidFetchBackgroundDownload/Source/AndroidFetchBackgroundDownload/Public/AndroidPlatformBackgroundHttpRequest.h
+++ b/Engine/Plugins/Online/Android/AndroidFetchBackgroundDownload/Source/AndroidFetchBackgroundDownload/Public/AndroidPlatformBackgroundHttpRequest.h
@@ -19,6 +19,7 @@ public:
 
 	virtual void PauseRequest() override;
     virtual void ResumeRequest() override;
+	virtual void CancelRequest() override;
     
 	FString ToJSon() const;
 
