Force worker requeue change.
- When losing internet we don't want to requeue the worker, doing this leads to a broken download notifications.

diff --git a/Engine/Plugins/Online/Android/AndroidFetchBackgroundDownload/Source/AndroidFetchBackgroundDownload/Private/AndroidPlatformBackgroundHttpManager.cpp b/Engine/Plugins/Online/Android/AndroidFetchBackgroundDownload/Source/AndroidFetchBackgroundDownload/Private/AndroidPlatformBackgroundHttpManager.cpp
--- a/Engine/Plugins/Online/Android/AndroidFetchBackgroundDownload/Source/AndroidFetchBackgroundDownload/Private/AndroidPlatformBackgroundHttpManager.cpp
+++ b/Engine/Plugins/Online/Android/AndroidFetchBackgroundDownload/Source/AndroidFetchBackgroundDownload/Private/AndroidPlatformBackgroundHttpManager.cpp
@@ -1134,7 +1134,11 @@ void FAndroidPlatformBackgroundHttpManager::OnNetworkConnectionChanged(ENetworkC
 
 		UE_LOG(LogBackgroundHttpManager, Display, TEXT("Flagging worker for requeue due to network connection re-establishing."));
 		//On reconnecting to the network flag our worker to force requeue to make sure our work is continued
-		FPlatformAtomics::InterlockedExchange(&bShouldForceWorkerRequeue, true);
+		//Don't requeue the worker if there is no Internet.
+		if (ConnectionType != ENetworkConnectionType::None)
+		{
+			FPlatformAtomics::InterlockedExchange(&bShouldForceWorkerRequeue, true);
+		}
 	}
 }
 
