Handle cellular network in background download
- When the App is alive and cellular is needed we show a "Waiting for cellular approval" notification
- When the app is killed you can accept cellular download from the notification tray.
- Still needs to rework cellular to connect it to the Java layer
- Accepting cellular when the app is open will replicated the value to the Java layer

diff --git a/Engine/Build/Android/Java/proguard-project.txt b/Engine/Build/Android/Java/proguard-project.txt
--- a/Engine/Build/Android/Java/proguard-project.txt
+++ b/Engine/Build/Android/Java/proguard-project.txt
@@ -46,6 +46,9 @@
 -keep class com.epicgames.unreal.LocalNotificationReceiver {
 	public *;
 }
+-keep class com.epicgames.unreal.CellularReceiver {
+	public *;
+}
 -keep class com.epicgames.unreal.MulticastBroadcastReceiver {
 	public *;
 }
diff --git a/Engine/Build/Android/Java/src/com/epicgames/unreal/CellularReceiver.java b/Engine/Build/Android/Java/src/com/epicgames/unreal/CellularReceiver.java
new file mode 100644
--- /dev/null
+++ b/Engine/Build/Android/Java/src/com/epicgames/unreal/CellularReceiver.java
@@ -0,0 +1,25 @@
+// Copyright Epic Games, Inc. All Rights Reserved.
+
+package com.epicgames.unreal;
+
+import com.epicgames.unreal.Logger;
+import android.app.Notification;
+import android.app.NotificationManager;
+import android.app.NotificationChannel;
+import android.app.PendingIntent;
+import android.content.BroadcastReceiver;
+import android.content.Context;
+import android.content.Intent;
+import androidx.core.app.NotificationCompat;
+import android.content.SharedPreferences;
+
+public class CellularReceiver extends BroadcastReceiver
+{
+	public void onReceive(Context context, Intent intent)
+	{
+		SharedPreferences preferences = context.getSharedPreferences("CellularNetworkPreferences", context.MODE_PRIVATE);
+		SharedPreferences.Editor editor = preferences.edit();
+		editor.putInt("AllowCellular", 1);
+		editor.commit();
+	}
+}
\ No newline at end of file
diff --git a/Engine/Build/Android/Java/src/com/epicgames/unreal/GameActivity.java.template b/Engine/Build/Android/Java/src/com/epicgames/unreal/GameActivity.java.template
--- a/Engine/Build/Android/Java/src/com/epicgames/unreal/GameActivity.java.template
+++ b/Engine/Build/Android/Java/src/com/epicgames/unreal/GameActivity.java.template
@@ -4677,6 +4677,22 @@ public class GameActivity extends $${gameActivitySuperClass}$$ implements Surfac
 		});
 	}
 
+	public void AndroidThunkJava_SetCellularPreference(int value)
+	{
+		Context context = getApplicationContext();
+		SharedPreferences preferences = context.getSharedPreferences("CellularNetworkPreferences", context.MODE_PRIVATE);
+		SharedPreferences.Editor editor = preferences.edit();
+		editor.putInt("AllowCellular", value);
+		editor.commit();
+	}
+
+	public int AndroidThunkJava_GetCellularPreference()
+	{
+		Context context = getApplicationContext();
+		SharedPreferences preferences = context.getSharedPreferences("CellularNetworkPreferences", context.MODE_PRIVATE);
+		return preferences.getInt("AllowCellular", 0 );
+	}
+
 	// call back into native code from the Java UI thread, initializing any available VR HMD modules
 	public void AndroidThunkJava_InitHMDs()
 	{
diff --git a/Engine/Build/Android/Java/src/com/epicgames/unreal/network/NetworkChangedManager.java b/Engine/Build/Android/Java/src/com/epicgames/unreal/network/NetworkChangedManager.java
--- a/Engine/Build/Android/Java/src/com/epicgames/unreal/network/NetworkChangedManager.java
+++ b/Engine/Build/Android/Java/src/com/epicgames/unreal/network/NetworkChangedManager.java
@@ -393,6 +393,12 @@ public final class NetworkChangedManager implements NetworkConnectivityClient {
 				break;
 		}
 	}
+
+	@Override
+	public NetworkTransportType networkTransportTypeCheck()
+	{
+		return calculateNetworkTransport(connectivityManager);
+	}
 	
 	@Override
 	public void checkConnectivity() {
diff --git a/Engine/Build/Android/Java/src/com/epicgames/unreal/network/NetworkConnectivityClient.java b/Engine/Build/Android/Java/src/com/epicgames/unreal/network/NetworkConnectivityClient.java
--- a/Engine/Build/Android/Java/src/com/epicgames/unreal/network/NetworkConnectivityClient.java
+++ b/Engine/Build/Android/Java/src/com/epicgames/unreal/network/NetworkConnectivityClient.java
@@ -47,4 +47,10 @@ public interface NetworkConnectivityClient {
 	 * Check for network connectivity
 	 */
 	void checkConnectivity();
+
+	/**
+	 * Check the current type of network connection
+	 * @return Type of network connection the device currently have
+	 */
+	NetworkTransportType networkTransportTypeCheck();
 }
diff --git a/Engine/Plugins/Online/Android/AndroidFetchBackgroundDownload/Source/AndroidFetchBackgroundDownload/Java/download/UEDownloadWorker.java b/Engine/Plugins/Online/Android/AndroidFetchBackgroundDownload/Source/AndroidFetchBackgroundDownload/Java/download/UEDownloadWorker.java
--- a/Engine/Plugins/Online/Android/AndroidFetchBackgroundDownload/Source/AndroidFetchBackgroundDownload/Java/download/UEDownloadWorker.java
+++ b/Engine/Plugins/Online/Android/AndroidFetchBackgroundDownload/Source/AndroidFetchBackgroundDownload/Java/download/UEDownloadWorker.java
@@ -10,6 +10,7 @@ import android.app.PendingIntent;
 import android.content.Context;
 import android.net.Uri;
 import android.os.Build;
+import android.content.SharedPreferences;
 
 import androidx.annotation.Nullable;
 import androidx.annotation.NonNull;
@@ -32,6 +33,8 @@ import com.epicgames.unreal.download.DownloadProgressListener;
 import com.epicgames.unreal.download.datastructs.DownloadQueueDescription;
 import com.epicgames.unreal.download.datastructs.DownloadWorkerParameterKeys;
 import com.epicgames.unreal.download.fetch.FetchManager;
+import com.epicgames.unreal.LocalNotificationReceiver;
+import com.epicgames.unreal.CellularReceiver;
 
 import com.epicgames.unreal.network.NetworkConnectivityClient;
 import com.epicgames.unreal.network.NetworkChangedManager;
@@ -80,6 +83,10 @@ public class UEDownloadWorker extends UEWorker implements DownloadProgressListen
 			CancelIntent = WorkManager.getInstance(getApplicationContext())
 				.createCancelPendingIntent(getId());
 		}
+		if (null == ApproveIntent)
+		{
+			ApproveIntent = new Intent(getApplicationContext(), CellularReceiver.class);
+		}
 		
 		//Generate our NotificationDescription so that we load important data from our InputData() to control notification content
 		if (null == NotificationDescription)
@@ -167,6 +174,45 @@ public class UEDownloadWorker extends UEWorker implements DownloadProgressListen
 			Log.debug("Finishing OnWorkerStart. CachedResult:" + CachedResult + " bReceivedResult:" + bReceivedResult);
 		}
 	}
+
+	// Checking what the current network type is.
+	// Will pause all downloads if the network type is Cellular
+	private void NetworkTypeCheck()
+	{
+		boolean ShouldHandleCellular = false;
+		if (NotificationDescription != null)
+		{
+			ShouldHandleCellular = NotificationDescription.ShouldHandleCellular;
+		}
+		// This should be handled by the Game thread otherwise
+		if (!bGameThreadIsActive && ShouldHandleCellular)
+		{
+			// Pause all downloads if the current cellular preference does not allow cellular downloading
+			NetworkConnectivityClient.NetworkTransportType networkType = NetworkChangedManager.getInstance().networkTransportTypeCheck();
+			if (!bWaitingForCellularApproval && networkType == NetworkConnectivityClient.NetworkTransportType.CELLULAR)
+			{
+				Context context = getApplicationContext();
+				SharedPreferences preferences = context.getSharedPreferences("CellularNetworkPreferences", context.MODE_PRIVATE);
+				boolean allowCell = (preferences.getInt("AllowCellular", 0) > 0);
+				if (!allowCell)
+				{
+					mFetchManager.PauseAllDownloads();
+					bWaitingForCellularApproval = true;
+				}
+			}
+			else if (bWaitingForCellularApproval)
+			{
+				Context context = getApplicationContext();
+				SharedPreferences preferences = context.getSharedPreferences("CellularNetworkPreferences", context.MODE_PRIVATE);
+				boolean allowCell = (preferences.getInt("AllowCellular", 0) > 0);
+				if (allowCell)
+				{
+					mFetchManager.ResumeAllDownloads();
+					bWaitingForCellularApproval = false;
+				}
+			}
+		}
+	}
 	
 	private void Tick(DownloadQueueDescription QueueDescription)
 	{
@@ -175,7 +221,7 @@ public class UEDownloadWorker extends UEWorker implements DownloadProgressListen
 		if (!bReceivedResult && bHasEnqueueHappened && !bForceStopped)
 		{
 			mFetchManager.RequestGroupProgressUpdate(QueueDescription.DownloadGroupID,  this);
-
+			NetworkTypeCheck();
 			//Keeping the code path to insert a watch dog later on
 			//mFetchManager.RequestCheckDownloadsStillActive(this);
 
@@ -259,7 +305,11 @@ public class UEDownloadWorker extends UEWorker implements DownloadProgressListen
 		}
 		
 		Notification notification = null;
-		if(!bLostNetwork)
+		if (bWaitingForCellularApproval)
+		{
+			notification = CreateCellularWaitNotification(context, Description, pendingNotificationIntent);
+		}
+		else if(!bLostNetwork)
 		{
 			notification = CreateDownloadProgressNotification(context, Description, pendingNotificationIntent);
 		}
@@ -303,6 +353,60 @@ public class UEDownloadWorker extends UEWorker implements DownloadProgressListen
 		return notification;
 	}
 
+	public Notification CreateCellularWaitNotification(Context context,
+		DownloadNotificationDescription Description,
+		PendingIntent pendingNotificationIntent)
+	{
+		//Setup Notification Text Values (ContentText and ContentInfo)
+		boolean bIsComplete = (Description.CurrentProgress >= Description.MAX_PROGRESS);
+		String NotificationTextToUse = null;
+		if (!bIsComplete)
+		{
+			NotificationTextToUse = String.format(Description.ContentText, Description.CurrentProgress);
+		}
+		else
+		{
+			NotificationTextToUse = Description.ContentCompleteText;
+		}
+		PendingIntent cellularNotificationIntent = null;
+		{
+			cellularNotificationIntent = PendingIntent.getBroadcast(context, Description.NotificationID, ApproveIntent, PendingIntent.FLAG_IMMUTABLE);
+		}
+		Notification notification;
+		if (bGameThreadIsActive)
+		{
+			notification = new NotificationCompat.Builder(context, Description.NotificationChannelID)
+				.setContentTitle(Description.WaitingForCellularText)
+				.setTicker(Description.WaitingForCellularText)
+				.setContentText(NotificationTextToUse)
+				.setContentIntent(pendingNotificationIntent)
+				.setProgress(Description.MAX_PROGRESS, Description.CurrentProgress, Description.Indeterminate)
+				.setOngoing(true)
+				.setOnlyAlertOnce (true)
+				.setSmallIcon(Description.SmallIconResourceID)
+				.setNotificationSilent()
+				.build();
+		}
+		else
+		{
+			notification = new NotificationCompat.Builder(context, Description.NotificationChannelID)
+				.setContentTitle(Description.WaitingForCellularText)
+				.setTicker(Description.WaitingForCellularText)
+				.setContentText(NotificationTextToUse)
+				.setContentIntent(pendingNotificationIntent)
+				.setProgress(Description.MAX_PROGRESS, Description.CurrentProgress, Description.Indeterminate)
+				.setOngoing(true)
+				.setOnlyAlertOnce (true)
+				.setSmallIcon(Description.SmallIconResourceID)
+				.addAction(Description.CancelIconResourceID, Description.ApproveText, cellularNotificationIntent)
+				.addAction(Description.CancelIconResourceID, Description.CancelText, CancelIntent)
+				.setNotificationSilent()
+				.build();
+		}
+		
+		return notification;
+	}
+
 	public Notification CreateDownloadProgressNotification(Context context,
 		DownloadNotificationDescription Description,
 		PendingIntent pendingNotificationIntent)
@@ -369,6 +473,16 @@ public class UEDownloadWorker extends UEWorker implements DownloadProgressListen
 	{
 		return (IsWorkEndTerminal());
 	}
+
+	private void ResetCellularPreference()
+	{
+		// Edit this to include Allways allow and only allow this time
+		Context context = getApplicationContext();
+		SharedPreferences preferences = context.getSharedPreferences("CellularNetworkPreferences", context.MODE_PRIVATE);
+		SharedPreferences.Editor editor = preferences.edit();
+		editor.putInt("AllowCellular", 0);
+		editor.commit();
+	}
 	
 	//
 	// DownloadCompletionListener Implementation
@@ -412,6 +526,9 @@ public class UEDownloadWorker extends UEWorker implements DownloadProgressListen
 			{
 				if (bDidAllRequestsSucceed)
 				{
+					//Resetting cellular preference for the next run, you may want to allow cellular for this run, but not the next one.
+					//This should be reworked and hooked up to the game code.
+					ResetCellularPreference();
 					SetWorkResult_Success();
 				}
 				//by default if UE didn't give us a behavior, lets just retry the download through the worker if one of the downloads failed
@@ -457,11 +574,26 @@ public class UEDownloadWorker extends UEWorker implements DownloadProgressListen
 	//
 	public void PauseRequest(String RequestID)
 	{
+		boolean ShouldHandleCellular = false;
+		if (NotificationDescription != null)
+		{
+			ShouldHandleCellular = NotificationDescription.ShouldHandleCellular;
+		}
+		if (ShouldHandleCellular)
+		{
+			NetworkConnectivityClient.NetworkTransportType networkType = NetworkChangedManager.getInstance().networkTransportTypeCheck();
+			if (networkType == NetworkConnectivityClient.NetworkTransportType.CELLULAR)
+			{
+				bWaitingForCellularApproval = true;
+			}
+		}
 		mFetchManager.PauseDownload(RequestID, true);
 	}
 	
 	public void ResumeRequest(String RequestID)
 	{
+		// When the C++ resumes the request all network issues should be solved
+		bWaitingForCellularApproval = false;
 		mFetchManager.PauseDownload(RequestID, false);
 	}
 	
@@ -469,6 +601,12 @@ public class UEDownloadWorker extends UEWorker implements DownloadProgressListen
 	{
 		mFetchManager.CancelDownload(RequestID);
 	}
+
+	// This is a way of figuring out if the worker was spawned when the app was killed or not
+	public static void AndroidThunkJava_GameThreadIsActive() 
+	{
+		bGameThreadIsActive = true;
+	}
 	
 	//Native functions used to bubble up progress to native UE code
 	public native void nativeAndroidBackgroundDownloadOnWorkerStart(String WorkID);
@@ -478,6 +616,7 @@ public class UEDownloadWorker extends UEWorker implements DownloadProgressListen
 	public native void nativeAndroidBackgroundDownloadOnAllComplete(boolean bDidAllRequestsSucceed);
 	public native void nativeAndroidBackgroundDownloadOnTick();
 	
+	private boolean bWaitingForCellularApproval = false;
 	private NetworkConnectivityClient.Listener NetworkListener = null;
 	private boolean bLostNetwork = false;
 	private boolean bForceStopped = false;
@@ -485,5 +624,7 @@ public class UEDownloadWorker extends UEWorker implements DownloadProgressListen
 	private volatile boolean bHasEnqueueHappened = false;
 	static volatile FetchManager mFetchManager;
 	private PendingIntent CancelIntent = null;
+	private Intent ApproveIntent = null;
 	private DownloadNotificationDescription NotificationDescription = null;
+	private static boolean bGameThreadIsActive = false;
 }
\ No newline at end of file
diff --git a/Engine/Plugins/Online/Android/AndroidFetchBackgroundDownload/Source/AndroidFetchBackgroundDownload/Java/download/datastructs/DownloadNotificationDescription.java b/Engine/Plugins/Online/Android/AndroidFetchBackgroundDownload/Source/AndroidFetchBackgroundDownload/Java/download/datastructs/DownloadNotificationDescription.java
--- a/Engine/Plugins/Online/Android/AndroidFetchBackgroundDownload/Source/AndroidFetchBackgroundDownload/Java/download/datastructs/DownloadNotificationDescription.java
+++ b/Engine/Plugins/Online/Android/AndroidFetchBackgroundDownload/Source/AndroidFetchBackgroundDownload/Java/download/datastructs/DownloadNotificationDescription.java
@@ -58,6 +58,11 @@ public class DownloadNotificationDescription
 			}
 		}
 
+		//Configuration 
+		{
+			ShouldHandleCellular = data.getBoolean(DownloadWorkerParameterKeys.CELLULAR_HANDLE_KEY, false);
+		}
+
 		//Load notification content information
 		{
 			TitleText = data.getString(DownloadWorkerParameterKeys.NOTIFICATION_CONTENT_TITLE_KEY);
@@ -89,6 +94,18 @@ public class DownloadNotificationDescription
 			{
 				NoInternetAvailable = "No Internet Available";
 			}
+
+			WaitingForCellularText = data.getString(DownloadWorkerParameterKeys.NOTIFICATION_CONTENT_WAITING_FOR_CELLULAR_TEXT_KEY);
+			if (null == WaitingForCellularText)
+			{
+				WaitingForCellularText = "Waiting for cellular approval";
+			}
+
+			ApproveText = data.getString(DownloadWorkerParameterKeys.NOTIFICATION_CONTENT_APPROVE_TEXT_KEY);
+			if (null == ApproveText)
+			{
+				ApproveText = "Approve";
+			}
 		}
 		
 		//Load the Cancel Icon Resource
@@ -228,11 +245,15 @@ public class DownloadNotificationDescription
 	
 	public int NotificationID = 0;
 
+	public boolean ShouldHandleCellular = false;
+
 	public String TitleText = null;
 	public String ContentText = null;
 	public String ContentCompleteText = null;
 	public String CancelText = null;
 	public String NoInternetAvailable = null;
+	public String WaitingForCellularText = null;
+	public String ApproveText = null;
 	
 	public int CancelIconResourceID = 0;
 	public int SmallIconResourceID = 0;
diff --git a/Engine/Plugins/Online/Android/AndroidFetchBackgroundDownload/Source/AndroidFetchBackgroundDownload/Java/download/datastructs/DownloadWorkerParameterKeys.java b/Engine/Plugins/Online/Android/AndroidFetchBackgroundDownload/Source/AndroidFetchBackgroundDownload/Java/download/datastructs/DownloadWorkerParameterKeys.java
--- a/Engine/Plugins/Online/Android/AndroidFetchBackgroundDownload/Source/AndroidFetchBackgroundDownload/Java/download/datastructs/DownloadWorkerParameterKeys.java
+++ b/Engine/Plugins/Online/Android/AndroidFetchBackgroundDownload/Source/AndroidFetchBackgroundDownload/Java/download/datastructs/DownloadWorkerParameterKeys.java
@@ -29,12 +29,16 @@ public class DownloadWorkerParameterKeys
 	public final static String NOTIFICATION_ID_KEY = "NotificationId";
 	//random value that our NOTIFICATION_ID is set to if not provided using the above key
 	public final static int NOTIFICATION_DEFAULT_ID_KEY = 1923901283;
+
+	public final static String CELLULAR_HANDLE_KEY = "CellularHandleKey";
 	
 	public final static String NOTIFICATION_CONTENT_TITLE_KEY = "NotificationContentTitle";
 	public final static String NOTIFICATION_CONTENT_TEXT_KEY = "NotificationContentText";
 	public final static String NOTIFICATION_CONTENT_CANCEL_DOWNLOAD_TEXT_KEY = "NotificationContentCancelDownloadText";
 	public final static String NOTIFICATION_CONTENT_NO_INTERNET_TEXT_KEY = "NotificationContentNoInternetDownloadText";
 	public final static String NOTIFICATION_CONTENT_COMPLETE_TEXT_KEY = "NotificationContentCompleteText";
+	public final static String NOTIFICATION_CONTENT_WAITING_FOR_CELLULAR_TEXT_KEY = "NotificationContentWaitingForCellularText";
+	public final static String NOTIFICATION_CONTENT_APPROVE_TEXT_KEY = "NotificationContentApproveText";
 
 	public final static String NOTIFICATION_RESOURCE_CANCEL_ICON_NAME = "NotificationResourceCancelIconName";
 	public final static String NOTIFICATION_RESOURCE_CANCEL_ICON_TYPE = "NotificationResourceCancelIconType";
diff --git a/Engine/Plugins/Online/Android/AndroidFetchBackgroundDownload/Source/AndroidFetchBackgroundDownload/Java/download/fetch/FetchManager.java b/Engine/Plugins/Online/Android/AndroidFetchBackgroundDownload/Source/AndroidFetchBackgroundDownload/Java/download/fetch/FetchManager.java
--- a/Engine/Plugins/Online/Android/AndroidFetchBackgroundDownload/Source/AndroidFetchBackgroundDownload/Java/download/fetch/FetchManager.java
+++ b/Engine/Plugins/Online/Android/AndroidFetchBackgroundDownload/Source/AndroidFetchBackgroundDownload/Java/download/fetch/FetchManager.java
@@ -38,6 +38,7 @@ import java.util.ArrayList;
 import java.util.HashMap;
 import java.util.List;
 import java.util.Set;
+import java.util.Map;
 
 import java.io.File;
 
@@ -256,6 +257,51 @@ public class FetchManager implements FetchDownloadProgressOwner, FetchEnqueueRes
 		}
 	}
 
+	public void PauseAllDownloads()
+	{
+		if (!IsFetchInstanceValid())
+		{
+			return;
+		}
+		
+		HashMap<String,DownloadDescription> TempRequestedDownloadsCopy =  new HashMap<String,DownloadDescription>(RequestedDownloads);
+		for(Map.Entry<String, DownloadDescription> entry : RequestedDownloads.entrySet()) 
+		{
+			
+    		DownloadDescription MatchedDesc = entry.getValue();
+			if (null == MatchedDesc)
+			{
+				return;
+			}
+			
+			MatchedDesc.bIsPaused = true;
+			FetchInstance.pause(MatchedDesc.CachedFetchID);
+		}
+	}
+
+	public void ResumeAllDownloads()
+	{
+		if (!IsFetchInstanceValid())
+		{
+			return;
+		}
+		
+		HashMap<String,DownloadDescription> TempRequestedDownloadsCopy =  new HashMap<String,DownloadDescription>(RequestedDownloads);
+		for(Map.Entry<String, DownloadDescription> entry : RequestedDownloads.entrySet()) 
+		{
+			
+    		DownloadDescription MatchedDesc = entry.getValue();
+			if (null == MatchedDesc)
+			{
+				return;
+			}
+			
+			MatchedDesc.bIsPaused = false;
+			FetchInstance.resume(MatchedDesc.CachedFetchID);
+		}
+	}
+
+
 	public void CancelDownload(String RequestID)
 	{
 		if (!IsFetchInstanceValid())
diff --git a/Engine/Plugins/Online/Android/AndroidFetchBackgroundDownload/Source/AndroidFetchBackgroundDownload/Private/AndroidPlatformBackgroundHttpManager.cpp b/Engine/Plugins/Online/Android/AndroidFetchBackgroundDownload/Source/AndroidFetchBackgroundDownload/Private/AndroidPlatformBackgroundHttpManager.cpp
--- a/Engine/Plugins/Online/Android/AndroidFetchBackgroundDownload/Source/AndroidFetchBackgroundDownload/Private/AndroidPlatformBackgroundHttpManager.cpp
+++ b/Engine/Plugins/Online/Android/AndroidFetchBackgroundDownload/Source/AndroidFetchBackgroundDownload/Private/AndroidPlatformBackgroundHttpManager.cpp
@@ -22,6 +22,8 @@
 
 #include "Misc/Paths.h"
 
+#include "HAL/IConsoleManager.h"
+
 #if UE_BUILD_SHIPPING
 // always clear any exceptions in shipping
 #define CHECK_JNI_RESULT(Id) if (Id == 0) { Env->ExceptionClear(); }
@@ -43,6 +45,12 @@
 
 #define LOCTEXT_NAMESPACE "AndroidBackgroundHttpManager"
 
+static TAutoConsoleVariable<bool> CVarEnableCellularHandling(
+	TEXT("bgdl.EnableCellularHandling"),
+	false,
+	TEXT("Enable/disable Cellular handling in the Java layer"),
+	ECVF_Default);
+
 //Call backs called by the JNI java systems to call back into C++ code in this AndroidPlatformBackgroundHttpManager
 namespace FAndroidBackgroundDownloadDelegates
 {
@@ -78,11 +86,15 @@ const FString FAndroidNativeDownloadWorkerParameterKeys::NOTIFICATION_ID_KEY = T
 //random value that our NOTIFICATION_ID is set to if not provided using the above key. KEEP IN SYNC WITH DownloadWorkerParameterKeys.java
 const int FAndroidNativeDownloadWorkerParameterKeys::NOTIFICATION_DEFAULT_ID_KEY = 1923901283;
 
+const FString FAndroidNativeDownloadWorkerParameterKeys::CELLULAR_HANDLE_KEY = TEXT("CellularHandleKey");
+
 const FString FAndroidNativeDownloadWorkerParameterKeys::NOTIFICATION_CONTENT_TITLE_KEY = TEXT("NotificationContentTitle");
 const FString FAndroidNativeDownloadWorkerParameterKeys::NOTIFICATION_CONTENT_TEXT_KEY = TEXT("NotificationContentText");
 const FString FAndroidNativeDownloadWorkerParameterKeys::NOTIFICATION_CONTENT_CANCEL_DOWNLOAD_TEXT_KEY = TEXT("NotificationContentCancelDownloadText");
 const FString FAndroidNativeDownloadWorkerParameterKeys::NOTIFICATION_CONTENT_NO_INTERNET_TEXT_KEY = TEXT("NotificationContentNoInternetDownloadText");
 const FString FAndroidNativeDownloadWorkerParameterKeys::NOTIFICATION_CONTENT_COMPLETE_TEXT_KEY = TEXT("NotificationContentCompleteText");
+const FString FAndroidNativeDownloadWorkerParameterKeys::NOTIFICATION_CONTENT_WAITING_FOR_CELLULAR_TEXT_KEY = TEXT("NotificationContentWaitingForCellularText");
+const FString FAndroidNativeDownloadWorkerParameterKeys::NOTIFICATION_CONTENT_APPROVE_TEXT_KEY = TEXT("NotificationContentApproveText");
 
 const FString FAndroidNativeDownloadWorkerParameterKeys::NOTIFICATION_RESOURCE_CANCEL_ICON_NAME = TEXT("NotificationResourceCancelIconName");
 const FString FAndroidNativeDownloadWorkerParameterKeys::NOTIFICATION_RESOURCE_CANCEL_ICON_TYPE = TEXT("NotificationResourceCancelIconType");
@@ -318,7 +330,9 @@ void FAndroidPlatformBackgroundHttpManager::ActivatePendingRequests()
 			const FString FileNameForDownloadDescList = GetFullFileNameForDownloadDescriptionList();
 			FScopedJavaObject<jstring> JavaFileNameString = FJavaHelper::ToJavaString(Env, FileNameForDownloadDescList);	
 			bool bDidWriteSucceed = Env->CallStaticBooleanMethod(FAndroidPlatformBackgroundHttpManager::JavaInfo.DownloadDescriptionClass, FAndroidPlatformBackgroundHttpManager::JavaInfo.WriteDownloadDescriptionListToFileMethod, *JavaFileNameString, *DescriptionArray);
-				
+			// Calling this to signal to the worker that it was spawned by a forgrounded app
+			Env->CallStaticVoidMethod(FAndroidPlatformBackgroundHttpManager::JavaInfo.UEDownloadWorkerClass, FAndroidPlatformBackgroundHttpManager::JavaInfo.GameThreadIsActive);	
+
 			bool bDidSuccessfullyScheduleWork = false;
 
 			//If we successfully created the JSON DownloadDescription file, 
@@ -337,11 +351,15 @@ void FAndroidPlatformBackgroundHttpManager::ActivatePendingRequests()
 				//Set our MaxActiveDownloads in the underlying java layer to match our expectation
 				WorkParams.AddDataToWorkerParameters(FAndroidNativeDownloadWorkerParameterKeys::DOWNLOAD_MAX_CONCURRENT_REQUESTS_KEY, MaxActiveDownloads);
 
+				WorkParams.AddDataToWorkerParameters(FAndroidNativeDownloadWorkerParameterKeys::CELLULAR_HANDLE_KEY, CVarEnableCellularHandling.GetValueOnGameThread());
+
 				//Make sure we pass in localized notification text bits for the important worker keys
 				WorkParams.AddDataToWorkerParameters(FAndroidNativeDownloadWorkerParameterKeys::NOTIFICATION_CONTENT_TITLE_KEY, AndroidBackgroundHTTPManagerDefaultLocalizedText.DefaultNotificationText_Title.GetText());
 				WorkParams.AddDataToWorkerParameters(FAndroidNativeDownloadWorkerParameterKeys::NOTIFICATION_CONTENT_COMPLETE_TEXT_KEY, AndroidBackgroundHTTPManagerDefaultLocalizedText.DefaultNotificationText_Complete.GetText());
 				WorkParams.AddDataToWorkerParameters(FAndroidNativeDownloadWorkerParameterKeys::NOTIFICATION_CONTENT_CANCEL_DOWNLOAD_TEXT_KEY, AndroidBackgroundHTTPManagerDefaultLocalizedText.DefaultNotificationText_Cancel.GetText());
 				WorkParams.AddDataToWorkerParameters(FAndroidNativeDownloadWorkerParameterKeys::NOTIFICATION_CONTENT_NO_INTERNET_TEXT_KEY, AndroidBackgroundHTTPManagerDefaultLocalizedText.DefaultNotificationText_NoInternet.GetText());
+				WorkParams.AddDataToWorkerParameters(FAndroidNativeDownloadWorkerParameterKeys::NOTIFICATION_CONTENT_WAITING_FOR_CELLULAR_TEXT_KEY, AndroidBackgroundHTTPManagerDefaultLocalizedText.DefaultNotificationText_WaitingForCellular.GetText());
+				WorkParams.AddDataToWorkerParameters(FAndroidNativeDownloadWorkerParameterKeys::NOTIFICATION_CONTENT_APPROVE_TEXT_KEY, AndroidBackgroundHTTPManagerDefaultLocalizedText.DefaultNotificationText_Approve.GetText());
 
 				//Expect our ContentText to have a {DownloadPercent} argument in it by default, so this will replace that with the Java string format argument so Java can insert the appropriate value
 				FFormatNamedArguments Arguments;
@@ -932,6 +950,12 @@ void FAndroidPlatformBackgroundHttpManager::FJavaClassInfo::Initialize()
 				CHECK_JNI_RESULT(DownloadDescriptionClass);
 			}
 
+			//find jmethodID for checking if the worker was created using the Android WorkManager
+			{
+				GameThreadIsActive = FJavaWrapper::FindStaticMethod(Env, UEDownloadWorkerClass, "AndroidThunkJava_GameThreadIsActive",  "()V", bIsOptional);
+				CHECK_JNI_RESULT(GameThreadIsActive);
+			}
+
 			//find jmethodID for DownloadDescription methods
 			{
 				//Grab a bunch of necessary JNI methods we will need to create our DownloadDescriptions
@@ -1010,6 +1034,8 @@ FAndroidPlatformBackgroundHttpManager::FAndroidBackgroundHTTPManagerDefaultLocal
 	, DefaultNotificationText_Complete()
 	, DefaultNotificationText_Cancel()
 	, DefaultNotificationText_NoInternet()
+	, DefaultNotificationText_WaitingForCellular()
+	, DefaultNotificationText_Approve()
 {
 }
 
@@ -1057,6 +1083,20 @@ void FAndroidPlatformBackgroundHttpManager::FAndroidBackgroundHTTPManagerDefault
 				ParsePolyglotTextItem(DefaultNotificationText_NoInternet, TEXT("Notification.NoInternetText"), NoInternetTextStrings);
 			}
 
+			//DefaultNotificationText_WaitingForCellular
+			{
+				TArray<FString> WaitingForCellularTextStrings;
+				Config->GetArray(*TextConfigSection, TEXT("DefaultNotificationText_WaitingForCellular"), WaitingForCellularTextStrings);
+				ParsePolyglotTextItem(DefaultNotificationText_WaitingForCellular, TEXT("Notification.WaitingForCellularText"), WaitingForCellularTextStrings);
+			}
+
+			//DefaultNotificationText_Approve
+			{
+				TArray<FString> ApproveTextStrings;
+				Config->GetArray(*TextConfigSection, TEXT("DefaultNotificationText_Approve"), ApproveTextStrings);
+				ParsePolyglotTextItem(DefaultNotificationText_Approve, TEXT("Notification.ApproveText"), ApproveTextStrings);
+			}
+
 		}
 	}
 
@@ -1066,6 +1106,8 @@ void FAndroidPlatformBackgroundHttpManager::FAndroidBackgroundHTTPManagerDefault
 	ForceValidPolyglotText(DefaultNotificationText_Complete, TEXT("DefaultNotificationText_Complete"));
 	ForceValidPolyglotText(DefaultNotificationText_Cancel, TEXT("DefaultNotificationText_Cancel"));
 	ForceValidPolyglotText(DefaultNotificationText_NoInternet, TEXT("DefaultNotificationText_NoInternet"));
+	ForceValidPolyglotText(DefaultNotificationText_WaitingForCellular, TEXT("DefaultNotificationText_WaitingForCellular"));
+	ForceValidPolyglotText(DefaultNotificationText_Approve, TEXT("DefaultNotificationText_Approve"));
 }
 
 void FAndroidPlatformBackgroundHttpManager::FAndroidBackgroundHTTPManagerDefaultLocalizedText::ForceValidPolyglotText(FPolyglotTextData& TextDataOut, const FString& DebugPolyglotTextName)
diff --git a/Engine/Plugins/Online/Android/AndroidFetchBackgroundDownload/Source/AndroidFetchBackgroundDownload/Public/AndroidPlatformBackgroundHttpManager.h b/Engine/Plugins/Online/Android/AndroidFetchBackgroundDownload/Source/AndroidFetchBackgroundDownload/Public/AndroidPlatformBackgroundHttpManager.h
--- a/Engine/Plugins/Online/Android/AndroidFetchBackgroundDownload/Source/AndroidFetchBackgroundDownload/Public/AndroidPlatformBackgroundHttpManager.h
+++ b/Engine/Plugins/Online/Android/AndroidFetchBackgroundDownload/Source/AndroidFetchBackgroundDownload/Public/AndroidPlatformBackgroundHttpManager.h
@@ -157,6 +157,8 @@ protected:
 		FPolyglotTextData DefaultNotificationText_Complete;
 		FPolyglotTextData DefaultNotificationText_Cancel;
 		FPolyglotTextData DefaultNotificationText_NoInternet;
+		FPolyglotTextData DefaultNotificationText_WaitingForCellular;
+		FPolyglotTextData DefaultNotificationText_Approve;
 
 
 		void InitFromIniSettings(const FString& ConfigFileName);
@@ -181,6 +183,7 @@ private:
 		jclass UEDownloadWorkerClass;
 		jclass DownloadDescriptionClass;
 
+		jmethodID GameThreadIsActive;
 		//Necessary JNI methods we will need to create our DownloadDescriptions
 		jmethodID CreateArrayStaticMethod;
 		jmethodID WriteDownloadDescriptionListToFileMethod;
@@ -191,6 +194,7 @@ private:
 		FJavaClassInfo()
 			: UEDownloadWorkerClass(0)
 			, DownloadDescriptionClass(0)
+			, GameThreadIsActive(0)
 			, CreateArrayStaticMethod(0)
 			, WriteDownloadDescriptionListToFileMethod(0)
 		{}
@@ -218,11 +222,15 @@ public:
 	static const FString NOTIFICATION_ID_KEY;
 	static const int NOTIFICATION_DEFAULT_ID_KEY;
 
+	static const FString CELLULAR_HANDLE_KEY;
+
 	static const FString NOTIFICATION_CONTENT_TITLE_KEY;
 	static const FString NOTIFICATION_CONTENT_TEXT_KEY;
 	static const FString NOTIFICATION_CONTENT_CANCEL_DOWNLOAD_TEXT_KEY;
 	static const FString NOTIFICATION_CONTENT_NO_INTERNET_TEXT_KEY;
 	static const FString NOTIFICATION_CONTENT_COMPLETE_TEXT_KEY;
+	static const FString NOTIFICATION_CONTENT_WAITING_FOR_CELLULAR_TEXT_KEY;
+	static const FString NOTIFICATION_CONTENT_APPROVE_TEXT_KEY;
 
 	static const FString NOTIFICATION_RESOURCE_CANCEL_ICON_NAME;
 	static const FString NOTIFICATION_RESOURCE_CANCEL_ICON_TYPE;
diff --git a/Engine/Source/Programs/UnrealBuildTool/Platform/Android/UEDeployAndroid.cs b/Engine/Source/Programs/UnrealBuildTool/Platform/Android/UEDeployAndroid.cs
--- a/Engine/Source/Programs/UnrealBuildTool/Platform/Android/UEDeployAndroid.cs
+++ b/Engine/Source/Programs/UnrealBuildTool/Platform/Android/UEDeployAndroid.cs
@@ -2936,6 +2936,7 @@ namespace UnrealBuildTool
 				Text.AppendLine(serviceLine);
 			}
 
+
 			// Declare the 8 Vulkan program compiling services.
 			Text.AppendLine(@"		<service android:name=""com.epicgames.unreal.psoservices.VulkanProgramService"" android:process="":psoprogramservice"" />");
 			// Declare the remaining 7 Vulkan program compiling services. (all derived from VulkanProgramService)
@@ -2950,6 +2951,7 @@ namespace UnrealBuildTool
 			Text.AppendLine("\t\t<receiver android:name=\"AlarmReceiver\" />");
 
 			Text.AppendLine("\t\t<receiver android:name=\"com.epicgames.unreal.LocalNotificationReceiver\" />");
+			Text.AppendLine("\t\t<receiver android:name=\"com.epicgames.unreal.CellularReceiver\" />");
 
 			if (bRestoreNotificationsOnReboot)
 			{
diff --git a/Engine/Source/Runtime/Core/Private/Android/AndroidPlatformMisc.cpp b/Engine/Source/Runtime/Core/Private/Android/AndroidPlatformMisc.cpp
--- a/Engine/Source/Runtime/Core/Private/Android/AndroidPlatformMisc.cpp
+++ b/Engine/Source/Runtime/Core/Private/Android/AndroidPlatformMisc.cpp
@@ -144,6 +144,10 @@ extern void AndroidThunkCpp_ForceQuit();
 
 extern void AndroidThunkCpp_SetOrientation(int32 Value);
 
+extern void AndroidThunkCpp_SetCellularPreference(int32 Value);
+
+extern int32 AndroidThunkCpp_GetCellularPreference();
+
 // From AndroidFile.cpp
 extern FString GFontPathBase;
 
@@ -3121,6 +3125,22 @@ void FAndroidMisc::SetDeviceOrientation(EDeviceScreenOrientation NewDeviceOrenta
 	SetAllowedDeviceOrientation(NewDeviceOrentation);
 }
 
+void FAndroidMisc::SetCellularPreference(int32 Value)
+{
+#if USE_ANDROID_JNI
+	AndroidThunkCpp_SetCellularPreference(Value);
+#endif // USE_ANDROID_JNI
+}
+
+int32 FAndroidMisc::GetCellularPreference()
+{
+	int32 value = 0;
+#if USE_ANDROID_JNI
+	value = AndroidThunkCpp_GetCellularPreference();
+#endif // USE_ANDROID_JNI
+	return value;
+}
+
 void FAndroidMisc::SetAllowedDeviceOrientation(EDeviceScreenOrientation NewAllowedDeviceOrientation)
 {
 	AllowedDeviceOrientation = NewAllowedDeviceOrientation;
diff --git a/Engine/Source/Runtime/Core/Public/Android/AndroidPlatformMisc.h b/Engine/Source/Runtime/Core/Public/Android/AndroidPlatformMisc.h
--- a/Engine/Source/Runtime/Core/Public/Android/AndroidPlatformMisc.h
+++ b/Engine/Source/Runtime/Core/Public/Android/AndroidPlatformMisc.h
@@ -52,6 +52,10 @@ struct CORE_API FAndroidMisc : public FGenericPlatformMisc
 	UE_DEPRECATED(5.1, "SetDeviceOrientation is deprecated. Use SetAllowedDeviceOrientation instead.")
 	static void SetDeviceOrientation(EDeviceScreenOrientation NewDeviceOrentation);
 	static void SetAllowedDeviceOrientation(EDeviceScreenOrientation NewAllowedDeviceOrientation);
+
+	// Change this to an Enum with Always allow, allow and deny
+	static void SetCellularPreference(int32 Value);
+	static int32 GetCellularPreference();
     
 	FORCEINLINE static int32 GetMaxPathLength()
 	{
diff --git a/Engine/Source/Runtime/Launch/Private/Android/AndroidJNI.cpp b/Engine/Source/Runtime/Launch/Private/Android/AndroidJNI.cpp
--- a/Engine/Source/Runtime/Launch/Private/Android/AndroidJNI.cpp
+++ b/Engine/Source/Runtime/Launch/Private/Android/AndroidJNI.cpp
@@ -125,6 +125,8 @@ void FJavaWrapper::FindClassesAndMethods(JNIEnv* Env)
 	AndroidThunkJava_GetIntentExtrasString = FindMethod(Env, GameActivityClassID, "AndroidThunkJava_GetIntentExtrasString", "(Ljava/lang/String;)Ljava/lang/String;", bIsOptional);
 	AndroidThunkJava_PushSensorEvents = FindMethod(Env, GameActivityClassID, "AndroidThunkJava_PushSensorEvents", "()V", bIsOptional);
 	AndroidThunkJava_SetOrientation = FindMethod(Env, GameActivityClassID, "AndroidThunkJava_SetOrientation", "(I)V", bIsOptional);
+	AndroidThunkJava_SetCellularPreference = FindMethod(Env, GameActivityClassID, "AndroidThunkJava_SetCellularPreference", "(I)V", bIsOptional);
+	AndroidThunkJava_GetCellularPreference = FindMethod(Env, GameActivityClassID, "AndroidThunkJava_GetCellularPreference", "()I", bIsOptional);
 
 	// Screen capture/recording permission
 	AndroidThunkJava_IsScreenCaptureDisabled = FindMethod(Env, GameActivityClassID, "AndroidThunkJava_IsScreenCaptureDisabled", "()Z", bIsOptional);
@@ -428,6 +430,8 @@ jmethodID FJavaWrapper::AndroidThunkJava_PushSensorEvents;
 jmethodID FJavaWrapper::AndroidThunkJava_IsScreenCaptureDisabled;
 jmethodID FJavaWrapper::AndroidThunkJava_DisableScreenCapture;
 jmethodID FJavaWrapper::AndroidThunkJava_SetOrientation;
+jmethodID FJavaWrapper::AndroidThunkJava_SetCellularPreference;
+jmethodID FJavaWrapper::AndroidThunkJava_GetCellularPreference;
 
 jclass FJavaWrapper::InputDeviceInfoClass;
 jfieldID FJavaWrapper::InputDeviceInfo_VendorId;
@@ -1336,6 +1340,24 @@ void AndroidThunkCpp_SetOrientation(int32 Value)
 	}
 }
 
+void AndroidThunkCpp_SetCellularPreference(int32 Value)
+{
+	if (JNIEnv* Env = FAndroidApplication::GetJavaEnv())
+	{
+		FJavaWrapper::CallVoidMethod(Env, FJavaWrapper::GameActivityThis, FJavaWrapper::AndroidThunkJava_SetCellularPreference, Value);
+	}
+}
+
+int32 AndroidThunkCpp_GetCellularPreference()
+{
+	int32 value = 0;
+	if (JNIEnv* Env = FAndroidApplication::GetJavaEnv())
+	{
+		FJavaWrapper::CallIntMethod(Env, FJavaWrapper::GameActivityThis, FJavaWrapper::AndroidThunkJava_GetCellularPreference);
+	}
+	return value;
+}
+
 bool AndroidThunkCpp_IsMusicActive()
 {
 	bool bIsActive = false;
diff --git a/Engine/Source/Runtime/Launch/Public/Android/AndroidJNI.h b/Engine/Source/Runtime/Launch/Public/Android/AndroidJNI.h
--- a/Engine/Source/Runtime/Launch/Public/Android/AndroidJNI.h
+++ b/Engine/Source/Runtime/Launch/Public/Android/AndroidJNI.h
@@ -73,6 +73,9 @@ public:
 	static jmethodID AndroidThunkJava_SetSustainedPerformanceMode;
 	static jmethodID AndroidThunkJava_PushSensorEvents;
 	static jmethodID AndroidThunkJava_SetOrientation;
+	static jmethodID AndroidThunkJava_SetCellularPreference;
+	static jmethodID AndroidThunkJava_GetCellularPreference;
+
 
 	// Screen capture/recording permission
 	static jmethodID AndroidThunkJava_IsScreenCaptureDisabled;
diff --git a/Engine/Source/Runtime/Online/BackgroundHTTP/Private/BackgroundHttpManagerImpl.cpp b/Engine/Source/Runtime/Online/BackgroundHTTP/Private/BackgroundHttpManagerImpl.cpp
--- a/Engine/Source/Runtime/Online/BackgroundHTTP/Private/BackgroundHttpManagerImpl.cpp
+++ b/Engine/Source/Runtime/Online/BackgroundHTTP/Private/BackgroundHttpManagerImpl.cpp
@@ -15,6 +15,9 @@
 #include "Stats/Stats.h"
 #include "ProfilingDebugging/CsvProfiler.h"
 
+#if PLATFORM_ANDROID
+#include "Android/AndroidPlatformMisc.h"
+#endif
 
 DEFINE_LOG_CATEGORY(LogBackgroundHttpManager);
 CSV_DECLARE_CATEGORY_MODULE_EXTERN(BACKGROUNDHTTP_API, BackgroundDownload);
@@ -344,6 +347,13 @@ void FBackgroundHttpManagerImpl::CleanUpDataAfterCompletingRequest(const FBackgr
 	}
 }
 
+void FBackgroundHttpManagerImpl::SetCellularPreference(int32 Value)
+{
+#if PLATFORM_ANDROID
+	FAndroidMisc::SetCellularPreference(Value);
+#endif
+}
+
 bool FBackgroundHttpManagerImpl::AssociateWithAnyExistingRequest(const FBackgroundHttpRequestPtr Request)
 {
 	bool bDidAssociateWithExistingRequest = false;
diff --git a/Engine/Source/Runtime/Online/BackgroundHTTP/Public/BackgroundHttpManagerImpl.h b/Engine/Source/Runtime/Online/BackgroundHTTP/Public/BackgroundHttpManagerImpl.h
--- a/Engine/Source/Runtime/Online/BackgroundHTTP/Public/BackgroundHttpManagerImpl.h
+++ b/Engine/Source/Runtime/Online/BackgroundHTTP/Public/BackgroundHttpManagerImpl.h
@@ -44,6 +44,7 @@ public:
 	virtual FString GetTempFileLocationForURL(const FString& URL) override;
 	
 	virtual void CleanUpDataAfterCompletingRequest(const FBackgroundHttpRequestPtr Request) override;
+	virtual void SetCellularPreference(int32 Value) override;
 	
 	//FTSTickerObjectBase implementation
 	virtual bool Tick(float DeltaTime) override;
diff --git a/Engine/Source/Runtime/Online/BackgroundHTTP/Public/Interfaces/IBackgroundHttpManager.h b/Engine/Source/Runtime/Online/BackgroundHTTP/Public/Interfaces/IBackgroundHttpManager.h
--- a/Engine/Source/Runtime/Online/BackgroundHTTP/Public/Interfaces/IBackgroundHttpManager.h
+++ b/Engine/Source/Runtime/Online/BackgroundHTTP/Public/Interfaces/IBackgroundHttpManager.h
@@ -75,6 +75,11 @@ public:
 	 * Should really never be calling this outside of a BackgroundHTTP class
 	 */
 	virtual void CleanUpDataAfterCompletingRequest(const FBackgroundHttpRequestPtr Request) = 0;
+
+	/**
+		* Setting cellular preference
+		*/
+	virtual void SetCellularPreference(int32 Value) = 0;
 	
 protected:
 	/**
