Adding localization for Cellular network aproval

diff --git a/Engine/Plugins/Online/Android/AndroidFetchBackgroundDownload/Config/BaseAndroidFetchBackgroundDownload.ini b/Engine/Plugins/Online/Android/AndroidFetchBackgroundDownload/Config/BaseAndroidFetchBackgroundDownload.ini
--- a/Engine/Plugins/Online/Android/AndroidFetchBackgroundDownload/Config/BaseAndroidFetchBackgroundDownload.ini
+++ b/Engine/Plugins/Online/Android/AndroidFetchBackgroundDownload/Config/BaseAndroidFetchBackgroundDownload.ini
@@ -83,4 +83,34 @@
 +DefaultNotificationText_NoInternet=("pl","Brak internetu")
 +DefaultNotificationText_NoInternet=("pt-BR","Internet indisponível")
 +DefaultNotificationText_NoInternet=("ru","Нет доступа в интернет")
-+DefaultNotificationText_NoInternet=("tr","İnternet bağlantısı yok")
\ No newline at end of file
++DefaultNotificationText_NoInternet=("tr","İnternet bağlantısı yok")
+
+;Waiting for cellular text
++DefaultNotificationText_WaitingForCellular=("native","Waiting for cellular approval")
++DefaultNotificationText_WaitingForCellular=("ar","في انتظار السماح باستخدام بيانات الهاتف")
++DefaultNotificationText_WaitingForCellular=("de","Warten auf Bestätigung der mobilen Datenverwendung")
++DefaultNotificationText_WaitingForCellular=("es","En espera de aprobación del uso de datos móviles")
++DefaultNotificationText_WaitingForCellular=("es-419","Esperando aprobación de uso de datos móviles")
++DefaultNotificationText_WaitingForCellular=("fr","En attente d'autorisation de l'utilisation des données mobiles")
++DefaultNotificationText_WaitingForCellular=("it","In attesa di autorizzazione per l’utilizzo dei dati mobili")
++DefaultNotificationText_WaitingForCellular=("ja","モバイルデータ通信の許可を待っています")
++DefaultNotificationText_WaitingForCellular=("ko","모바일 데이터 사용 승인 기다리는 중")
++DefaultNotificationText_WaitingForCellular=("pl","Oczekiwanie na zatwierdzenie użycia danych komórkowych")
++DefaultNotificationText_WaitingForCellular=("pt-BR","Aguardando a autorização de uso de dados móveis")
++DefaultNotificationText_WaitingForCellular=("ru","Ожидание разрешения на использование сотовой сети для передачи данных")
++DefaultNotificationText_WaitingForCellular=("tr","Hücresel veri onayı bekleniyor")
+
+;Approve Text
++DefaultNotificationText_Approve=("native","Approve")
++DefaultNotificationText_Approve=("ar","السماح")
++DefaultNotificationText_Approve=("de","Bestätigen")
++DefaultNotificationText_Approve=("es","Aprobar")
++DefaultNotificationText_Approve=("es-419","Aprobar")
++DefaultNotificationText_Approve=("fr","Autoriser")
++DefaultNotificationText_Approve=("it","Autorizza")
++DefaultNotificationText_Approve=("ja","許可する")
++DefaultNotificationText_Approve=("ko","승인")
++DefaultNotificationText_Approve=("pl","Zatwierdź")
++DefaultNotificationText_Approve=("pt-BR","Autorizar")
++DefaultNotificationText_Approve=("ru","Разрешить")
++DefaultNotificationText_Approve=("tr","Onayla")
