Recovery support if changing background download enable on Android

#android

diff --git a/Engine/Plugins/Online/Android/AndroidFetchBackgroundDownload/Source/AndroidFetchBackgroundDownload/Private/AndroidPlatformBackgroundHttpManager.cpp b/Engine/Plugins/Online/Android/AndroidFetchBackgroundDownload/Source/AndroidFetchBackgroundDownload/Private/AndroidPlatformBackgroundHttpManager.cpp
--- a/Engine/Plugins/Online/Android/AndroidFetchBackgroundDownload/Source/AndroidFetchBackgroundDownload/Private/AndroidPlatformBackgroundHttpManager.cpp
+++ b/Engine/Plugins/Online/Android/AndroidFetchBackgroundDownload/Source/AndroidFetchBackgroundDownload/Private/AndroidPlatformBackgroundHttpManager.cpp
@@ -363,7 +363,7 @@ void FAndroidPlatformBackgroundHttpManager::ActivatePendingRequests()
 
 				//Expect our ContentText to have a {DownloadPercent} argument in it by default, so this will replace that with the Java string format argument so Java can insert the appropriate value
 				FFormatNamedArguments Arguments;
-				Arguments.Emplace(TEXT("DownloadPercent"), FText::FromString(TEXT("%03d%%")));
+				Arguments.Emplace(TEXT("DownloadPercent"), FText::FromString(TEXT("%3d%%")));
 				FText UpdatedContentText = FText::Format(AndroidBackgroundHTTPManagerDefaultLocalizedText.DefaultNotificationText_Content.GetText(), Arguments);
 				WorkParams.AddDataToWorkerParameters(FAndroidNativeDownloadWorkerParameterKeys::NOTIFICATION_CONTENT_TEXT_KEY, UpdatedContentText);
 
