Optional disable progress percentage in notification on Android

#android

diff --git a/Engine/Plugins/Online/Android/AndroidFetchBackgroundDownload/Source/AndroidFetchBackgroundDownload/Java/download/UEDownloadWorker.java b/Engine/Plugins/Online/Android/AndroidFetchBackgroundDownload/Source/AndroidFetchBackgroundDownload/Java/download/UEDownloadWorker.java
--- a/Engine/Plugins/Online/Android/AndroidFetchBackgroundDownload/Source/AndroidFetchBackgroundDownload/Java/download/UEDownloadWorker.java
+++ b/Engine/Plugins/Online/Android/AndroidFetchBackgroundDownload/Source/AndroidFetchBackgroundDownload/Java/download/UEDownloadWorker.java
@@ -326,23 +326,34 @@ public class UEDownloadWorker extends UEWorker implements DownloadProgressListen
 		PendingIntent pendingNotificationIntent)
 	{
 		//Setup Notification Text Values (ContentText and ContentInfo)
+		boolean bShowPercentage = ShouldShowPercentage(context);
 		boolean bIsComplete = (Description.CurrentProgress >= Description.MAX_PROGRESS);
 		String NotificationTextToUse = null;
 		if (!bIsComplete)
 		{
-			NotificationTextToUse = String.format(Description.ContentText, Description.CurrentProgress);
+			if (bShowPercentage)
+			{
+				NotificationTextToUse = String.format(Description.ContentText, Description.CurrentProgress);
+			}
+			else
+			{
+				NotificationTextToUse = Description.ContentText.replace("%3d%%", "");
+			}
 		}
 		else
 		{
 			NotificationTextToUse = Description.ContentCompleteText;
 		}
 
+		int CurrentProgress = bShowPercentage ? Description.CurrentProgress : 0;
+		boolean Indeterminate = bShowPercentage ? Description.Indeterminate : true;
+
 		Notification notification = new NotificationCompat.Builder(context, Description.NotificationChannelID)
 			.setContentTitle(Description.NoInternetAvailable)
 			.setTicker(Description.TitleText)
 			.setContentText(NotificationTextToUse)
 			.setContentIntent(pendingNotificationIntent)
-			.setProgress(Description.MAX_PROGRESS, Description.CurrentProgress, Description.Indeterminate)
+			.setProgress(Description.MAX_PROGRESS, CurrentProgress, Indeterminate)
 			.setOngoing(true)
 			.setOnlyAlertOnce (true)
 			.setSmallIcon(Description.SmallIconResourceID)
@@ -358,11 +369,19 @@ public class UEDownloadWorker extends UEWorker implements DownloadProgressListen
 		PendingIntent pendingNotificationIntent)
 	{
 		//Setup Notification Text Values (ContentText and ContentInfo)
+		boolean bShowPercentage = ShouldShowPercentage(context);
 		boolean bIsComplete = (Description.CurrentProgress >= Description.MAX_PROGRESS);
 		String NotificationTextToUse = null;
 		if (!bIsComplete)
 		{
-			NotificationTextToUse = String.format(Description.ContentText, Description.CurrentProgress);
+			if (bShowPercentage)
+			{
+				NotificationTextToUse = String.format(Description.ContentText, Description.CurrentProgress);
+			}
+			else
+			{
+				NotificationTextToUse = Description.ContentText.replace("%3d%%", "");
+			}
 		}
 		else
 		{
@@ -372,6 +391,10 @@ public class UEDownloadWorker extends UEWorker implements DownloadProgressListen
 		{
 			cellularNotificationIntent = PendingIntent.getBroadcast(context, Description.NotificationID, ApproveIntent, PendingIntent.FLAG_IMMUTABLE);
 		}
+
+		int CurrentProgress = bShowPercentage ? Description.CurrentProgress : 0;
+		boolean Indeterminate = bShowPercentage ? Description.Indeterminate : true;
+
 		Notification notification;
 		if (bGameThreadIsActive)
 		{
@@ -380,7 +403,7 @@ public class UEDownloadWorker extends UEWorker implements DownloadProgressListen
 				.setTicker(Description.WaitingForCellularText)
 				.setContentText(NotificationTextToUse)
 				.setContentIntent(pendingNotificationIntent)
-				.setProgress(Description.MAX_PROGRESS, Description.CurrentProgress, Description.Indeterminate)
+				.setProgress(Description.MAX_PROGRESS, CurrentProgress, Indeterminate)
 				.setOngoing(true)
 				.setOnlyAlertOnce (true)
 				.setSmallIcon(Description.SmallIconResourceID)
@@ -394,7 +417,7 @@ public class UEDownloadWorker extends UEWorker implements DownloadProgressListen
 				.setTicker(Description.WaitingForCellularText)
 				.setContentText(NotificationTextToUse)
 				.setContentIntent(pendingNotificationIntent)
-				.setProgress(Description.MAX_PROGRESS, Description.CurrentProgress, Description.Indeterminate)
+				.setProgress(Description.MAX_PROGRESS, CurrentProgress, Indeterminate)
 				.setOngoing(true)
 				.setOnlyAlertOnce (true)
 				.setSmallIcon(Description.SmallIconResourceID)
@@ -412,23 +435,34 @@ public class UEDownloadWorker extends UEWorker implements DownloadProgressListen
 		PendingIntent pendingNotificationIntent)
 	{
 		//Setup Notification Text Values (ContentText and ContentInfo)
+		boolean bShowPercentage = ShouldShowPercentage(context);
 		boolean bIsComplete = (Description.CurrentProgress >= Description.MAX_PROGRESS);
 		String NotificationTextToUse = null;
 		if (!bIsComplete)
 		{
-			NotificationTextToUse = String.format(Description.ContentText, Description.CurrentProgress);
+			if (bShowPercentage)
+			{
+				NotificationTextToUse = String.format(Description.ContentText, Description.CurrentProgress);
+			}
+			else
+			{
+				NotificationTextToUse = Description.ContentText.replace("%3d%%", "");
+			}
 		}
 		else
 		{
 			NotificationTextToUse = Description.ContentCompleteText;
 		}
 
+		int CurrentProgress = bShowPercentage ? Description.CurrentProgress : 0;
+		boolean Indeterminate = bShowPercentage ? Description.Indeterminate : true;
+
 		Notification notification = new NotificationCompat.Builder(context, Description.NotificationChannelID)
 			.setContentTitle(Description.TitleText)
 			.setTicker(Description.TitleText)
 			.setContentText(NotificationTextToUse)
 			.setContentIntent(pendingNotificationIntent)
-			.setProgress(Description.MAX_PROGRESS, Description.CurrentProgress, Description.Indeterminate)
+			.setProgress(Description.MAX_PROGRESS, CurrentProgress, Indeterminate)
 			.setOngoing(true)
 			.setOnlyAlertOnce (true)
 			.setSmallIcon(Description.SmallIconResourceID)
@@ -474,6 +508,17 @@ public class UEDownloadWorker extends UEWorker implements DownloadProgressListen
 		return (IsWorkEndTerminal());
 	}
 
+	private boolean ShouldShowPercentage(Context context)
+	{
+		if (ShowPercentage == ESelectState.Unset)
+		{
+			SharedPreferences preferences = context.getSharedPreferences("BackgroundPreferences", context.MODE_PRIVATE);
+			boolean bShow = preferences.getBoolean("bShowPercentage", true);
+			ShowPercentage = bShow ? ESelectState.Enable : ESelectState.Disable;
+		}
+		return ShowPercentage == ESelectState.Enable;
+	}
+
 	private void ResetCellularPreference()
 	{
 		// Edit this to include Allways allow and only allow this time
@@ -627,4 +672,12 @@ public class UEDownloadWorker extends UEWorker implements DownloadProgressListen
 	private Intent ApproveIntent = null;
 	private DownloadNotificationDescription NotificationDescription = null;
 	private static boolean bGameThreadIsActive = false;
+
+	private enum ESelectState
+	{
+		Unset,
+		Disable,
+		Enable
+	};
+	private ESelectState ShowPercentage = ESelectState.Unset;
 }
\ No newline at end of file
