Download throughput calculation engine side
Implemented independent logic for each platform Android and iOS
No platforms were affected besides the target ones

diff --git a/Engine/Plugins/Online/Android/AndroidFetchBackgroundDownload/Source/AndroidFetchBackgroundDownload/Java/download/DownloadProgressListener.java b/Engine/Plugins/Online/Android/AndroidFetchBackgroundDownload/Source/AndroidFetchBackgroundDownload/Java/download/DownloadProgressListener.java
--- a/Engine/Plugins/Online/Android/AndroidFetchBackgroundDownload/Source/AndroidFetchBackgroundDownload/Java/download/DownloadProgressListener.java
+++ b/Engine/Plugins/Online/Android/AndroidFetchBackgroundDownload/Source/AndroidFetchBackgroundDownload/Java/download/DownloadProgressListener.java
@@ -9,6 +9,7 @@ public interface DownloadProgressListener
 	public void OnDownloadProgress(String RequestID, long BytesWrittenSinceLastCall, long TotalBytesWritten);
 	public void OnDownloadGroupProgress(int GroupID, int Progress, boolean Indeterminate);
 	public void OnDownloadComplete(String RequestID, String CompleteLocation, EDownloadCompleteReason CompleteReason);
+	public void OnDownloadMetrics(String RequestID, long TotalBytesDownloaded, long DownloadDuration);
 	public void OnAllDownloadsComplete(boolean bDidAllRequestsSucceed);
 	public void OnDownloadEnqueued(String RequestID, boolean bEnqueueSuccess);
 }
\ No newline at end of file
diff --git a/Engine/Plugins/Online/Android/AndroidFetchBackgroundDownload/Source/AndroidFetchBackgroundDownload/Java/download/UEDownloadWorker.java b/Engine/Plugins/Online/Android/AndroidFetchBackgroundDownload/Source/AndroidFetchBackgroundDownload/Java/download/UEDownloadWorker.java
--- a/Engine/Plugins/Online/Android/AndroidFetchBackgroundDownload/Source/AndroidFetchBackgroundDownload/Java/download/UEDownloadWorker.java
+++ b/Engine/Plugins/Online/Android/AndroidFetchBackgroundDownload/Source/AndroidFetchBackgroundDownload/Java/download/UEDownloadWorker.java
@@ -552,6 +552,12 @@ public class UEDownloadWorker extends UEWorker implements DownloadProgressListen
 		boolean bWasSuccess = (CompleteReason == EDownloadCompleteReason.Success);
 		nativeAndroidBackgroundDownloadOnComplete(RequestID, CompleteLocation, bWasSuccess);
 	}
+
+	@Override
+	public void OnDownloadMetrics(String RequestID, long TotalBytesDownloaded, long DownloadDuration)
+	{
+		nativeAndroidBackgroundDownloadOnMetrics(RequestID, TotalBytesDownloaded, DownloadDuration);
+	}
 	
 	@Override
 	public void OnAllDownloadsComplete(boolean bDidAllRequestsSucceed)
@@ -658,6 +664,7 @@ public class UEDownloadWorker extends UEWorker implements DownloadProgressListen
 	public native void nativeAndroidBackgroundDownloadOnWorkerStop(String WorkID);
 	public native void nativeAndroidBackgroundDownloadOnProgress(String TaskID, long BytesWrittenSinceLastCall, long TotalBytesWritten);
 	public native void nativeAndroidBackgroundDownloadOnComplete(String TaskID, String CompleteLocation, boolean bWasSuccess);
+	public native void nativeAndroidBackgroundDownloadOnMetrics(String TaskID, long TotalBytesDownloaded, long DownloadDuration);
 	public native void nativeAndroidBackgroundDownloadOnAllComplete(boolean bDidAllRequestsSucceed);
 	public native void nativeAndroidBackgroundDownloadOnTick();
 	
diff --git a/Engine/Plugins/Online/Android/AndroidFetchBackgroundDownload/Source/AndroidFetchBackgroundDownload/Java/download/datastructs/DownloadDescription.java b/Engine/Plugins/Online/Android/AndroidFetchBackgroundDownload/Source/AndroidFetchBackgroundDownload/Java/download/datastructs/DownloadDescription.java
--- a/Engine/Plugins/Online/Android/AndroidFetchBackgroundDownload/Source/AndroidFetchBackgroundDownload/Java/download/datastructs/DownloadDescription.java
+++ b/Engine/Plugins/Online/Android/AndroidFetchBackgroundDownload/Source/AndroidFetchBackgroundDownload/Java/download/datastructs/DownloadDescription.java
@@ -55,6 +55,12 @@ public class DownloadDescription
 	public long TotalBytesNeeded = 0;
 	public long TotalDownloadedBytes = 0;
 	public DownloadProgressListener ProgressListener = null;
+
+	// Fields used to calculate the total download time
+	public long DownloadStartTime = 0;
+	public long DownloadPauseTime = 0;
+	public long TotalPausedTime = 0;
+
 	public DownloadDescription()
 	{
 		URLs = new ArrayList<String>();
diff --git a/Engine/Plugins/Online/Android/AndroidFetchBackgroundDownload/Source/AndroidFetchBackgroundDownload/Java/download/fetch/FetchDownloadProgressOwner.java b/Engine/Plugins/Online/Android/AndroidFetchBackgroundDownload/Source/AndroidFetchBackgroundDownload/Java/download/fetch/FetchDownloadProgressOwner.java
--- a/Engine/Plugins/Online/Android/AndroidFetchBackgroundDownload/Source/AndroidFetchBackgroundDownload/Java/download/fetch/FetchDownloadProgressOwner.java
+++ b/Engine/Plugins/Online/Android/AndroidFetchBackgroundDownload/Source/AndroidFetchBackgroundDownload/Java/download/fetch/FetchDownloadProgressOwner.java
@@ -14,6 +14,7 @@ import com.tonyodev.fetch2.FetchGroup;
 public interface FetchDownloadProgressOwner
 {
 	public void OnDownloadQueued(@NonNull Download download);
+	public void OnDownloadStarted(@NonNull Download download);
 	public void OnDownloadProgress(@NonNull Download download, boolean indeterminate, long downloadedBytesPerSecond, long etaInMilliSeconds);
 	public void OnDownloadChangePauseState(@NonNull Download download, boolean bIsPaused);
 	public void OnDownloadGroupProgress(@NonNull FetchGroup Group, DownloadProgressListener ProgressListener);
diff --git a/Engine/Plugins/Online/Android/AndroidFetchBackgroundDownload/Source/AndroidFetchBackgroundDownload/Java/download/fetch/FetchManager.java b/Engine/Plugins/Online/Android/AndroidFetchBackgroundDownload/Source/AndroidFetchBackgroundDownload/Java/download/fetch/FetchManager.java
--- a/Engine/Plugins/Online/Android/AndroidFetchBackgroundDownload/Source/AndroidFetchBackgroundDownload/Java/download/fetch/FetchManager.java
+++ b/Engine/Plugins/Online/Android/AndroidFetchBackgroundDownload/Source/AndroidFetchBackgroundDownload/Java/download/fetch/FetchManager.java
@@ -646,6 +646,25 @@ public class FetchManager implements FetchDownloadProgressOwner, FetchEnqueueRes
 		OnDownloadChangePauseState(download, false);
 	}
 
+	@Override
+	public void OnDownloadStarted(@NonNull Download download)
+	{
+		String DownloadRequestID = GetRequestID(download);
+		Log.verbose("OnDownloadStarted: " + DownloadRequestID);
+
+		DownloadDescription MatchedDownload = RequestedDownloads.get(DownloadRequestID);
+		if (null == MatchedDownload)
+		{
+			Log.error("OnDownloadStarted called from DownloadProgressOwner implementation with a download that doesn't match any DownloadDesc! Download's Tag: " + DownloadRequestID);
+			return;
+		}
+
+		if (MatchedDownload.DownloadStartTime == 0)
+		{
+			MatchedDownload.DownloadStartTime = System.currentTimeMillis();
+		}
+	}
+
 	@Override
 	public void OnDownloadProgress(@NonNull Download download, boolean indeterminate, long downloadedBytesPerSecond, long etaInMilliSeconds)
 	{
@@ -714,10 +733,21 @@ public class FetchManager implements FetchDownloadProgressOwner, FetchEnqueueRes
 			
 			if (MatchedDownload.bIsPaused)
 			{
+				MatchedDownload.DownloadPauseTime = System.currentTimeMillis();
+
 				FetchInstance.pause(download.getId());
 			}
 			else
 			{
+				if (MatchedDownload.DownloadPauseTime != 0)
+				{
+					long ElapsedPauseTime = System.currentTimeMillis() - MatchedDownload.DownloadPauseTime;
+
+					MatchedDownload.TotalPausedTime += ElapsedPauseTime;
+
+					MatchedDownload.DownloadPauseTime = 0;
+				}
+
 				FetchInstance.resume(download.getId());
 			}
 		}
@@ -1143,7 +1173,18 @@ public class FetchManager implements FetchDownloadProgressOwner, FetchEnqueueRes
 			FailedDownloads.put(DownloadDesc.RequestID, DownloadDesc);
 		}
 
-		DownloadDesc.ProgressListener.OnDownloadComplete(DownloadDesc.RequestID, DownloadDesc.DestinationLocation, CompleteReasonToSend);	
+		long DownloadDuration = -1;
+		if (DownloadDesc.DownloadStartTime > 0)
+		{
+			long ElapsedDownloadTime = System.currentTimeMillis() - DownloadDesc.DownloadStartTime;
+
+			DownloadDuration = ElapsedDownloadTime - DownloadDesc.TotalPausedTime;
+		}
+
+		DownloadDesc.ProgressListener.OnDownloadMetrics(DownloadDesc.RequestID, DownloadDesc.TotalDownloadedBytes, DownloadDuration);
+
+		DownloadDesc.ProgressListener.OnDownloadComplete(DownloadDesc.RequestID, DownloadDesc.DestinationLocation, CompleteReasonToSend);
+		
 		if (CheckForAllCompleted)
 		{
 			CheckForAllDownloadsComplete(DownloadDesc.ProgressListener);
diff --git a/Engine/Plugins/Online/Android/AndroidFetchBackgroundDownload/Source/AndroidFetchBackgroundDownload/Java/download/fetch/FetchRequestProgressListener.java b/Engine/Plugins/Online/Android/AndroidFetchBackgroundDownload/Source/AndroidFetchBackgroundDownload/Java/download/fetch/FetchRequestProgressListener.java
--- a/Engine/Plugins/Online/Android/AndroidFetchBackgroundDownload/Source/AndroidFetchBackgroundDownload/Java/download/fetch/FetchRequestProgressListener.java
+++ b/Engine/Plugins/Online/Android/AndroidFetchBackgroundDownload/Source/AndroidFetchBackgroundDownload/Java/download/fetch/FetchRequestProgressListener.java
@@ -39,6 +39,11 @@ final class FetchRequestProgressListener implements FetchListener
 	{
 		Owner.OnDownloadQueued(download);
 	}
+
+	public void BroadcastDownloadStarted(@NonNull Download download)
+	{
+		Owner.OnDownloadStarted(download);
+	}
 	
 	public void BroadcastProgress(@NonNull Download download, boolean indeterminate)
 	{
@@ -90,6 +95,7 @@ final class FetchRequestProgressListener implements FetchListener
 	{
 		//Make sure to broadcast a change in pause state, because if we previously thought we were paused we aren't anymore!
 		Log.verbose("onStarted");
+		BroadcastDownloadStarted(download);
 		BroadcastChangePauseState(download, false);
 	}
 	
diff --git a/Engine/Plugins/Online/Android/AndroidFetchBackgroundDownload/Source/AndroidFetchBackgroundDownload/Private/AndroidPlatformBackgroundHttpManager.cpp b/Engine/Plugins/Online/Android/AndroidFetchBackgroundDownload/Source/AndroidFetchBackgroundDownload/Private/AndroidPlatformBackgroundHttpManager.cpp
--- a/Engine/Plugins/Online/Android/AndroidFetchBackgroundDownload/Source/AndroidFetchBackgroundDownload/Private/AndroidPlatformBackgroundHttpManager.cpp
+++ b/Engine/Plugins/Online/Android/AndroidFetchBackgroundDownload/Source/AndroidFetchBackgroundDownload/Private/AndroidPlatformBackgroundHttpManager.cpp
@@ -58,6 +58,7 @@ namespace FAndroidBackgroundDownloadDelegates
 	DECLARE_DELEGATE_TwoParams(FAndroidBackgroundDownload_OnWorkerStop, FString /*WorkID*/, jobject /*UnderlyingWorker*/);
 	DECLARE_DELEGATE_FourParams(FAndroidBackgroundDownload_OnProgress, jobject /*UnderlyingWorker*/, FString /*RequestID*/, int64_t /*BytesWrittenSinceLastCall*/, int64_t /*TotalBytesWritten*/);
 	DECLARE_DELEGATE_FourParams(FAndroidBackgroundDownload_OnComplete, jobject /*UnderlyingWorker*/, FString /*RequestID*/, FString /*CompleteLocation*/, bool /*bWasSuccess*/);
+	DECLARE_DELEGATE_FourParams(FAndroidBackgroundDownload_OnMetrics, jobject /*UnderlyingWorker*/, FString /*RequestID*/, int64_t /*TotalBytesDownloaded*/, int64_t /*DownloadDuration*/);
 	DECLARE_DELEGATE_TwoParams(FAndroidBackgroundDownload_OnAllComplete, jobject /*UnderlyingWorker*/, bool /*bDidAllRequestsSucceed*/);
 	DECLARE_DELEGATE_TwoParams(FAndroidBackgroundDownload_OnTickWorkerThread, JNIEnv*, jobject /*UnderlyingWorker*/);
 
@@ -66,6 +67,7 @@ namespace FAndroidBackgroundDownloadDelegates
 	FAndroidBackgroundDownload_OnWorkerStop AndroidBackgroundDownload_OnWorkerStop;
 	FAndroidBackgroundDownload_OnProgress AndroidBackgroundDownload_OnProgress;
 	FAndroidBackgroundDownload_OnComplete AndroidBackgroundDownload_OnComplete;
+	FAndroidBackgroundDownload_OnMetrics AndroidBackgroundDownload_OnMetrics;
 	FAndroidBackgroundDownload_OnAllComplete AndroidBackgroundDownload_OnAllComplete;
 	FAndroidBackgroundDownload_OnTickWorkerThread AndroidBackgroundDownload_OnTickWorkerThread;
 };
@@ -178,6 +180,7 @@ void FAndroidPlatformBackgroundHttpManager::Initialize()
 	FAndroidBackgroundDownloadDelegates::AndroidBackgroundDownload_OnWorkerStop.BindSP(this, &FAndroidPlatformBackgroundHttpManager::Java_OnWorkerStop);
 	FAndroidBackgroundDownloadDelegates::AndroidBackgroundDownload_OnProgress.BindSP(this, &FAndroidPlatformBackgroundHttpManager::Java_OnDownloadProgress);
 	FAndroidBackgroundDownloadDelegates::AndroidBackgroundDownload_OnComplete.BindSP(this, &FAndroidPlatformBackgroundHttpManager::Java_OnDownloadComplete);
+	FAndroidBackgroundDownloadDelegates::AndroidBackgroundDownload_OnMetrics.BindSP(this, &FAndroidPlatformBackgroundHttpManager::Java_OnDownloadMetrics);
 	FAndroidBackgroundDownloadDelegates::AndroidBackgroundDownload_OnAllComplete.BindSP(this, &FAndroidPlatformBackgroundHttpManager::Java_OnAllDownloadsComplete);
 	FAndroidBackgroundDownloadDelegates::AndroidBackgroundDownload_OnTickWorkerThread.BindSP(this, &FAndroidPlatformBackgroundHttpManager::Java_OnTick);
 
@@ -705,6 +708,19 @@ void FAndroidPlatformBackgroundHttpManager::Java_OnDownloadComplete(jobject Unde
 	}
 }
 
+void FAndroidPlatformBackgroundHttpManager::Java_OnDownloadMetrics(jobject UnderlyingWorker, FString RequestID, int64_t TotalBytesDownloaded, int64_t DownloadDuration)
+{
+	UE_LOG(LogBackgroundHttpManager, Log, TEXT("DownloadMetrics... RequestID:%s TotalBytesDownloaded:%ld DownloadDuration:%ld"), *RequestID, TotalBytesDownloaded, DownloadDuration);
+
+	FAndroidBackgroundHttpRequestPtr CompletedRequest = FindRequestByID(RequestID);
+	if (CompletedRequest.IsValid())
+	{
+		float DownloadDurationInSeconds = (float)((double)DownloadDuration / 1000.0f);
+
+		CompletedRequest->NotifyRequestMetricsAvailable((int32)TotalBytesDownloaded, DownloadDurationInSeconds);
+	}
+}
+
 void FAndroidPlatformBackgroundHttpManager::MarkUnderlyingJavaRequestAsCompleted(FBackgroundHttpRequestPtr Request, bool bSuccess /*= true*/)
 {
 	FAndroidBackgroundHttpRequestPtr AndroidRequest = StaticCastSharedPtr<FAndroidPlatformBackgroundHttpRequest>(Request);
@@ -1236,6 +1252,15 @@ JNI_METHOD void Java_com_epicgames_unreal_download_UEDownloadWorker_nativeAndroi
 	FAndroidBackgroundDownloadDelegates::AndroidBackgroundDownload_OnComplete.ExecuteIfBound(thiz, RequestID, ConvertedCompleteLocation, ConvertedbWasSuccess);
 }
 
+JNI_METHOD void Java_com_epicgames_unreal_download_UEDownloadWorker_nativeAndroidBackgroundDownloadOnMetrics(JNIEnv* jenv, jobject thiz, jstring TaskID, jlong TotalBytesDownloaded, jlong DownloadDuration)
+{
+	FString RequestID = FJavaHelper::FStringFromParam(jenv, TaskID);
+	int64_t ConvertedTotalBytesDownloaded = static_cast<uint64_t>(TotalBytesDownloaded);
+	int64_t ConvertedDownloadDuration = static_cast<uint64_t>(DownloadDuration);
+
+	FAndroidBackgroundDownloadDelegates::AndroidBackgroundDownload_OnMetrics.ExecuteIfBound(thiz, RequestID, ConvertedTotalBytesDownloaded, ConvertedDownloadDuration);
+}
+
 JNI_METHOD void Java_com_epicgames_unreal_download_UEDownloadWorker_nativeAndroidBackgroundDownloadOnAllComplete(JNIEnv* jenv, jobject thiz, jboolean bDidAllRequestsSucceed)
 {
 	bool ConvertedbDidAllRequestsSucceed = static_cast<bool>(bDidAllRequestsSucceed);
diff --git a/Engine/Plugins/Online/Android/AndroidFetchBackgroundDownload/Source/AndroidFetchBackgroundDownload/Public/AndroidPlatformBackgroundHttpManager.h b/Engine/Plugins/Online/Android/AndroidFetchBackgroundDownload/Source/AndroidFetchBackgroundDownload/Public/AndroidPlatformBackgroundHttpManager.h
--- a/Engine/Plugins/Online/Android/AndroidFetchBackgroundDownload/Source/AndroidFetchBackgroundDownload/Public/AndroidPlatformBackgroundHttpManager.h
+++ b/Engine/Plugins/Online/Android/AndroidFetchBackgroundDownload/Source/AndroidFetchBackgroundDownload/Public/AndroidPlatformBackgroundHttpManager.h
@@ -74,6 +74,7 @@ protected:
 	void Java_OnWorkerStop(FString WorkID, jobject UnderlyingWorker);
 	void Java_OnDownloadProgress(jobject UnderlyingWorker, FString RequestID, int64_t BytesWrittenSinceLastCall, int64_t TotalBytesWritten);
 	void Java_OnDownloadComplete(jobject UnderlyingWorker, FString RequestID, FString CompleteLocation, bool bWasSuccess);
+	void Java_OnDownloadMetrics(jobject UnderlyingWorker, FString RequestID, int64_t TotalBytesDownloaded, int64_t DownloadDuration);
 	void Java_OnAllDownloadsComplete(jobject UnderlyingWorker, bool bDidAllRequestsSucceed);
 	void Java_OnTick(JNIEnv* Env, jobject UnderlyingWorker);
 
diff --git a/Engine/Source/Runtime/ApplicationCore/Private/IOS/IOSBackgroundURLSessionHandler.cpp b/Engine/Source/Runtime/ApplicationCore/Private/IOS/IOSBackgroundURLSessionHandler.cpp
--- a/Engine/Source/Runtime/ApplicationCore/Private/IOS/IOSBackgroundURLSessionHandler.cpp
+++ b/Engine/Source/Runtime/ApplicationCore/Private/IOS/IOSBackgroundURLSessionHandler.cpp
@@ -936,8 +936,8 @@ static constexpr NSInteger HTTPStatusCodeErrorServer = 500;
 	}
 }
 
-#if !UE_BUILD_SHIPPING
 - (void)URLSession:(NSURLSession *)session task:(NSURLSessionTask *)task didFinishCollectingMetrics:(NSURLSessionTaskMetrics *)metrics {
+#if !UE_BUILD_SHIPPING
 	NSDateFormatter *formatter = [[NSDateFormatter alloc] init];
     [formatter setDateFormat:@"yyyy-MM-dd'T'HH:mm:ss.SSSZZZZZ"];
     [formatter setLocale:[NSLocale localeWithLocaleIdentifier:@"en_US_POSIX"]];
@@ -954,9 +954,10 @@ static constexpr NSInteger HTTPStatusCodeErrorServer = 500;
     // Get redirect count
     NSInteger redirectCount = metrics.transactionMetrics.count - 1; // The first metric is the initial request, redirects are additional metrics
     UE_DNLD_LOG(@"Redirect Count: %ld", (long)redirectCount);
-
+#endif
 
 	for (NSURLSessionTaskTransactionMetrics *metric in metrics.transactionMetrics) {
+#if !UE_BUILD_SHIPPING
 		UE_DNLD_LOG(@"Network Protocol Name: %@", metric.networkProtocolName);
 		UE_DNLD_LOG(@"Reused Connection: %@", metric.reusedConnection ? @"Yes" : @"No");
 		UE_DNLD_LOG(@"Proxy Connection: %@", metric.proxyConnection ? @"Yes": @"No");
@@ -966,28 +967,42 @@ static constexpr NSInteger HTTPStatusCodeErrorServer = 500;
         UE_DNLD_LOG(@"Response Start Date: %@", [formatter stringFromDate:metric.responseStartDate]);
         UE_DNLD_LOG(@"Request End Date: %@", [formatter stringFromDate:metric.requestEndDate]);
         UE_DNLD_LOG(@"Response End Date: %@", [formatter stringFromDate:metric.responseEndDate]);
+#endif
 
 		// Calculate and log response duration
         if (metric.responseStartDate && metric.responseEndDate) {
             NSTimeInterval responseDuration = [metric.responseEndDate timeIntervalSinceDate:metric.responseStartDate];
+#if !UE_BUILD_SHIPPING
             UE_DNLD_LOG(@"Response Duration: %f seconds", responseDuration);
+#endif
 
 			// Calculate and log download speed
             if (responseDuration > 0 && metric.countOfResponseBodyBytesReceived > 0) {
+				
+				FBackgroundURLSessionHandler::OnDownloadMetrics.Broadcast((uint64)task.taskIdentifier, (int32)metric.countOfResponseBodyBytesReceived, (float)responseDuration);
+				
+#if !UE_BUILD_SHIPPING
 				// bytes per second
-                double downloadSpeed = (double)metric.countOfResponseBodyBytesReceived / responseDuration;
-                NSString *formattedSpeed = [self formattedSpeed:downloadSpeed];
-                UE_DNLD_LOG(@"Download Speed: %@", formattedSpeed);
+				double downloadSpeed = (double)metric.countOfResponseBodyBytesReceived / responseDuration;
+				NSString *formattedSpeed = [self formattedSpeed:downloadSpeed];
+				UE_DNLD_LOG(@"Download Speed: %@", formattedSpeed);
             } else {
                 UE_DNLD_LOG(@"Download Speed: Not Available");
+#endif
             }
-        } else {
+        }
+#if !UE_BUILD_SHIPPING
+		else {
             UE_DNLD_LOG(@"Response Duration: Not Available");
         }
+#endif
     }
+#if !UE_BUILD_SHIPPING
 	UE_DNLD_LOG(@"-------------------------");
+#endif
 }
 
+#if !UE_BUILD_SHIPPING
 - (NSString *)formattedSpeed:(double)speedInBytesPerSecond {
     NSArray *units = @[@"bytes/second", @"KB/s", @"MB/s", @"GB/s"];
     double speed = speedInBytesPerSecond;
@@ -1534,6 +1549,7 @@ static constexpr NSInteger HTTPStatusCodeErrorServer = 500;
 const uint64 FBackgroundURLSessionHandler::InvalidDownloadId = [FBackgroundNSURLSession GetInvalidDownloadId];
 
 FBackgroundURLSessionHandler::FOnDownloadCompleted FBackgroundURLSessionHandler::OnDownloadCompleted;
+FBackgroundURLSessionHandler::FOnDownloadMetrics FBackgroundURLSessionHandler::OnDownloadMetrics;
 
 FBackgroundURLSessionHandler::FOnDownloadsCompletedWhileAppWasNotRunning FBackgroundURLSessionHandler::OnDownloadsCompletedWhileAppWasNotRunning;
 
diff --git a/Engine/Source/Runtime/ApplicationCore/Public/IOS/IOSBackgroundURLSessionHandler.h b/Engine/Source/Runtime/ApplicationCore/Public/IOS/IOSBackgroundURLSessionHandler.h
--- a/Engine/Source/Runtime/ApplicationCore/Public/IOS/IOSBackgroundURLSessionHandler.h
+++ b/Engine/Source/Runtime/ApplicationCore/Public/IOS/IOSBackgroundURLSessionHandler.h
@@ -59,6 +59,10 @@ public:
 	// Will be invoked from didFinishDownloadingToURL or didCompleteWithError.
 	DECLARE_MULTICAST_DELEGATE_TwoParams(FOnDownloadCompleted, const uint64 /*DownloadId*/, const bool /*bSuccess*/);
 	static FOnDownloadCompleted OnDownloadCompleted;
+	
+	// Will be invoked from didFinishCollectingMetrics
+	DECLARE_MULTICAST_DELEGATE_ThreeParams(FOnDownloadMetrics, const uint64 /*DownloadId*/, const int32 /*TotalBytesDownloaded*/, const float /*DownloadDuration*/)
+	static FOnDownloadMetrics OnDownloadMetrics;
 
 	// Will be invoked from handleEventsForBackgroundURLSession application delegate. Needs to be registered very early, e.g. from static constructor.
 	// handleEventsForBackgroundURLSession is only invoked if app was killed by OS while in background and then relaunched to notify that downloads were completed.
diff --git a/Engine/Source/Runtime/InstallBundleManager/Private/InstallBundleManagerInterface.cpp b/Engine/Source/Runtime/InstallBundleManager/Private/InstallBundleManagerInterface.cpp
--- a/Engine/Source/Runtime/InstallBundleManager/Private/InstallBundleManagerInterface.cpp
+++ b/Engine/Source/Runtime/InstallBundleManager/Private/InstallBundleManagerInterface.cpp
@@ -5,6 +5,8 @@
 
 FInstallBundleManagerInitCompleteMultiDelegate IInstallBundleManager::InitCompleteDelegate;
 
+FInstallBundleChunkDownloadMetricsMultiDelegate IInstallBundleManager::InstallBundleChunkDownloadMetricsDelegate;
+
 FInstallBundleCompleteMultiDelegate IInstallBundleManager::InstallBundleCompleteDelegate;
 
 FInstallBundlePausedMultiDelegate IInstallBundleManager::PausedBundleDelegate;
diff --git a/Engine/Source/Runtime/InstallBundleManager/Public/InstallBundleManagerInterface.h b/Engine/Source/Runtime/InstallBundleManager/Public/InstallBundleManagerInterface.h
--- a/Engine/Source/Runtime/InstallBundleManager/Public/InstallBundleManagerInterface.h
+++ b/Engine/Source/Runtime/InstallBundleManager/Public/InstallBundleManagerInterface.h
@@ -88,11 +88,19 @@ enum class EInstallBundleManagerInitErrorHandlerResult
 	StopInitialization, // Stop trying to initialize
 };
 
+struct FInstallBundleChunkDownloadInfo
+{
+	FString ChunkName;
+	int32 ChunkFileSize;
+	float ChunkDownloadDuration;
+};
+
 using FInstallBundleSourceOrCache = TUnion<FInstallBundleSourceType, FName>;
 
 DECLARE_DELEGATE_RetVal_OneParam(EInstallBundleManagerInitErrorHandlerResult, FInstallBundleManagerInitErrorHandler, EInstallBundleManagerInitResult);
 DECLARE_MULTICAST_DELEGATE_OneParam(FInstallBundleManagerInitCompleteMultiDelegate, EInstallBundleManagerInitResult);
 
+DECLARE_MULTICAST_DELEGATE_OneParam(FInstallBundleChunkDownloadMetricsMultiDelegate, FInstallBundleChunkDownloadInfo);
 DECLARE_MULTICAST_DELEGATE_OneParam(FInstallBundleCompleteMultiDelegate, FInstallBundleRequestResultInfo);
 DECLARE_MULTICAST_DELEGATE_OneParam(FInstallBundlePausedMultiDelegate, FInstallBundlePauseInfo);
 DECLARE_MULTICAST_DELEGATE_OneParam(FInstallBundleReleasedMultiDelegate, FInstallBundleReleaseRequestResultInfo);
@@ -110,6 +118,7 @@ class IInstallBundleManager : public TSharedFromThis<IInstallBundleManager>
 public:
 	static INSTALLBUNDLEMANAGER_API FInstallBundleManagerInitCompleteMultiDelegate InitCompleteDelegate;
 
+	static INSTALLBUNDLEMANAGER_API FInstallBundleChunkDownloadMetricsMultiDelegate InstallBundleChunkDownloadMetricsDelegate; // Called when a content request metrics is available
 	static INSTALLBUNDLEMANAGER_API FInstallBundleCompleteMultiDelegate InstallBundleCompleteDelegate; // Called when a content request is complete
 	static INSTALLBUNDLEMANAGER_API FInstallBundlePausedMultiDelegate PausedBundleDelegate;
 	static INSTALLBUNDLEMANAGER_API FInstallBundleReleasedMultiDelegate ReleasedDelegate; // Called when content release request is complete
diff --git a/Engine/Source/Runtime/Online/BackgroundHTTP/Private/BackgroundHttpRequestImpl.cpp b/Engine/Source/Runtime/Online/BackgroundHTTP/Private/BackgroundHttpRequestImpl.cpp
--- a/Engine/Source/Runtime/Online/BackgroundHTTP/Private/BackgroundHttpRequestImpl.cpp
+++ b/Engine/Source/Runtime/Online/BackgroundHTTP/Private/BackgroundHttpRequestImpl.cpp
@@ -73,6 +73,12 @@ void FBackgroundHttpRequestImpl::OnBackgroundDownloadComplete()
 	
 	//Now that we have called our completion delegates and done everything else for this request, allow us to clean up data for it
 	FBackgroundHttpModule::Get().GetBackgroundHttpManager()->CleanUpDataAfterCompletingRequest(SharedThis(this));
+	
+	if (OptionalMetricsInfo.IsSet())
+	{
+		OnRequestMetrics().ExecuteIfBound(SharedThis(this), OptionalMetricsInfo->TotalBytesDownloaded, OptionalMetricsInfo->DownloadDuration);
+		OptionalMetricsInfo.Reset();
+	}
 }
 
 void FBackgroundHttpRequestImpl::NotifyNotificationObjectOfComplete(bool bWasSuccess)
@@ -92,6 +98,14 @@ void FBackgroundHttpRequestImpl::NotifyNotificationObjectOfComplete(bool bWasSuc
 	}
 }
 
+void FBackgroundHttpRequestImpl::NotifyRequestMetricsAvailable(const int32 TotalBytesDownloaded, const float DownloadDuration)
+{
+	OptionalMetricsInfo = {
+		.TotalBytesDownloaded = TotalBytesDownloaded,
+		.DownloadDuration = DownloadDuration
+	};
+}
+
 void FBackgroundHttpRequestImpl::SetURLAsList(const TArray<FString>& URLs, int NumRetriesIn)
 {
 	NumberOfTotalRetries = NumRetriesIn;
@@ -133,6 +147,11 @@ FBackgroundHttpProgressUpdateDelegate& FBackgroundHttpRequestImpl::OnProgressUpd
 	return HttpProgressUpdateDelegate;
 }
 
+FBackgroundHttpRequestMetricsDelegate& FBackgroundHttpRequestImpl::OnRequestMetrics()
+{
+	return HttpRequestMetricsDelegate;
+}
+
 const FBackgroundHttpResponsePtr FBackgroundHttpRequestImpl::GetResponse() const
 {
 	return Response;
diff --git a/Engine/Source/Runtime/Online/BackgroundHTTP/Private/IOS/ApplePlatformBackgroundHttpManager.cpp b/Engine/Source/Runtime/Online/BackgroundHTTP/Private/IOS/ApplePlatformBackgroundHttpManager.cpp
--- a/Engine/Source/Runtime/Online/BackgroundHTTP/Private/IOS/ApplePlatformBackgroundHttpManager.cpp
+++ b/Engine/Source/Runtime/Online/BackgroundHTTP/Private/IOS/ApplePlatformBackgroundHttpManager.cpp
@@ -24,11 +24,13 @@ FApplePlatformBackgroundHttpManager::FApplePlatformBackgroundHttpManager()
 	, PendingRemoveRequestLock()
 {
 	OnDownloadCompletedHandle = FBackgroundURLSessionHandler::OnDownloadCompleted.AddRaw(this, &FApplePlatformBackgroundHttpManager::OnDownloadCompleted);
+	OnDownloadMetricsHandle = FBackgroundURLSessionHandler::OnDownloadMetrics.AddRaw(this, &FApplePlatformBackgroundHttpManager::OnDownloadMetrics);
 }
 
 FApplePlatformBackgroundHttpManager::~FApplePlatformBackgroundHttpManager()
 {
 	FBackgroundURLSessionHandler::OnDownloadCompleted.Remove(OnDownloadCompletedHandle);
+	FBackgroundURLSessionHandler::OnDownloadMetrics.Remove(OnDownloadMetricsHandle);
 }
 
 void FApplePlatformBackgroundHttpManager::AddRequest(const FBackgroundHttpRequestPtr GenericRequest)
@@ -134,4 +136,21 @@ void FApplePlatformBackgroundHttpManager::OnDownloadCompleted(const uint64 Downl
 			}
 		}
 	}
-}
\ No newline at end of file
+}
+
+void FApplePlatformBackgroundHttpManager::OnDownloadMetrics(const uint64 DownloadId, const int32 TotalBytesDownloaded, const float DownloadDuration)
+{
+	FRWScopeLock ScopeLock(ActiveRequestLock, SLT_ReadOnly);
+	for (FBackgroundHttpRequestPtr& GenericRequest : ActiveRequests)
+	{
+		FAppleBackgroundHttpRequestPtr Request = StaticCastSharedPtr<FApplePlatformBackgroundHttpRequest>(GenericRequest);
+		if (ensureAlwaysMsgf(Request.IsValid(), TEXT("Invalid Request Pointer in ActiveRequests list!")))
+		{
+			if (Request->GetInternalDownloadId() == DownloadId)
+			{
+				Request->NotifyRequestMetricsAvailable(TotalBytesDownloaded, DownloadDuration);
+				return;
+			}
+		}
+	}
+}
diff --git a/Engine/Source/Runtime/Online/BackgroundHTTP/Public/BackgroundHttpRequestImpl.h b/Engine/Source/Runtime/Online/BackgroundHTTP/Public/BackgroundHttpRequestImpl.h
--- a/Engine/Source/Runtime/Online/BackgroundHTTP/Public/BackgroundHttpRequestImpl.h
+++ b/Engine/Source/Runtime/Online/BackgroundHTTP/Public/BackgroundHttpRequestImpl.h
@@ -35,6 +35,7 @@ public:
 	BACKGROUNDHTTP_API virtual void CompleteWithExistingResponseData(FBackgroundHttpResponsePtr BackgroundResponse) override;
 	BACKGROUNDHTTP_API virtual FBackgroundHttpRequestCompleteDelegate& OnProcessRequestComplete() override;
 	BACKGROUNDHTTP_API virtual FBackgroundHttpProgressUpdateDelegate& OnProgressUpdated() override;
+	BACKGROUNDHTTP_API virtual FBackgroundHttpRequestMetricsDelegate& OnRequestMetrics() override;
 	BACKGROUNDHTTP_API virtual const FBackgroundHttpResponsePtr GetResponse() const override;
 	BACKGROUNDHTTP_API virtual const FString& GetRequestID() const override;
 	BACKGROUNDHTTP_API virtual void SetRequestID(const FString& NewRequestID) override;
@@ -42,7 +43,9 @@ public:
 	BACKGROUNDHTTP_API virtual EBackgroundHTTPPriority GetRequestPriority() const override;
 	BACKGROUNDHTTP_API virtual void SetRequestPriority(EBackgroundHTTPPriority NewPriority) override;
 	
-	BACKGROUNDHTTP_API virtual void NotifyNotificationObjectOfComplete(bool bWasSuccess);
+	BACKGROUNDHTTP_API virtual void NotifyNotificationObjectOfComplete(const bool bWasSuccess);
+	BACKGROUNDHTTP_API virtual void NotifyRequestMetricsAvailable(const int32 TotalBytesDownloaded, const float DownloadDuration);
+	
 protected:
 	TSharedPtr<FBackgroundHttpNotificationObject, ESPMode::ThreadSafe> DownloadCompleteNotificationObject;
 	FBackgroundHttpResponsePtr Response;
@@ -51,6 +54,15 @@ protected:
 	int NumberOfTotalRetries;
 	EBackgroundHTTPPriority RequestPriority;
 	
+	struct FDownloadMetricsInfo
+	{
+		int32 TotalBytesDownloaded;
+		float DownloadDuration;
+	};
+	
+	TOptional<FDownloadMetricsInfo> OptionalMetricsInfo;
+	
 	FBackgroundHttpRequestCompleteDelegate HttpRequestCompleteDelegate;
 	FBackgroundHttpProgressUpdateDelegate HttpProgressUpdateDelegate;
+	FBackgroundHttpRequestMetricsDelegate HttpRequestMetricsDelegate;
 };
diff --git a/Engine/Source/Runtime/Online/BackgroundHTTP/Public/IOS/ApplePlatformBackgroundHttpManager.h b/Engine/Source/Runtime/Online/BackgroundHTTP/Public/IOS/ApplePlatformBackgroundHttpManager.h
--- a/Engine/Source/Runtime/Online/BackgroundHTTP/Public/IOS/ApplePlatformBackgroundHttpManager.h
+++ b/Engine/Source/Runtime/Online/BackgroundHTTP/Public/IOS/ApplePlatformBackgroundHttpManager.h
@@ -35,5 +35,7 @@ private:
 	TArray<FBackgroundHttpRequestPtr> PendingRemoveRequests;
 	FRWLock PendingRemoveRequestLock;
 	FDelegateHandle OnDownloadCompletedHandle;
+	FDelegateHandle OnDownloadMetricsHandle;
 	void OnDownloadCompleted(const uint64 DownloadId, const bool bSuccess);
+	void OnDownloadMetrics(const uint64 DownloadId, const int32 TotalBytesDownloaded, const float DownloadDuration);
 };
diff --git a/Engine/Source/Runtime/Online/BackgroundHTTP/Public/Interfaces/IBackgroundHttpRequest.h b/Engine/Source/Runtime/Online/BackgroundHTTP/Public/Interfaces/IBackgroundHttpRequest.h
--- a/Engine/Source/Runtime/Online/BackgroundHTTP/Public/Interfaces/IBackgroundHttpRequest.h
+++ b/Engine/Source/Runtime/Online/BackgroundHTTP/Public/Interfaces/IBackgroundHttpRequest.h
@@ -30,6 +30,14 @@ DECLARE_DELEGATE_TwoParams(FBackgroundHttpRequestCompleteDelegate, FBackgroundHt
 */
 DECLARE_DELEGATE_ThreeParams(FBackgroundHttpProgressUpdateDelegate, FBackgroundHttpRequestPtr /*Request*/, int32 /*TotalBytesWritten*/, int32 /*BytesWrittenSinceLastUpdate*/);
 
+/**
+ * Delegate called when a Background Http request metrics are ready
+ * @param Request Background Http Request for which the metrics were calculated
+ * @param TotalBytesDownloaded The count of bytes downloaded in this request
+ * @param DownloadDuration The duration of the download in seconds
+ */
+DECLARE_DELEGATE_ThreeParams(FBackgroundHttpRequestMetricsDelegate, FBackgroundHttpRequestPtr /* Request */, int32 /* TotalBytesDownloaded */, float /* DownloadDuration */);
+
 /**
  * Enum used to describe download priority. Higher priorities will be downloaded first.
  * Note: Should always be kept in High -> Low priority order if adding more Priorities!
@@ -134,6 +142,11 @@ public:
 	* Delegate called when the request has a progress update.
 	*/
 	virtual FBackgroundHttpProgressUpdateDelegate& OnProgressUpdated() = 0;
+	
+	/**
+	 * Delegate called when the request's metrics were calculated
+	 */
+	virtual FBackgroundHttpRequestMetricsDelegate& OnRequestMetrics() = 0;
 
 	/**
 	* Called to begin processing the request.
