Update Fetch2 repository to new location now that jcenter no longer hosts this.

diff --git a/Engine/Plugins/Online/Android/AndroidFetchBackgroundDownload/Source/AndroidFetchBackgroundDownload/AndroidFetchBackgroundDownload_UPL.xml b/Engine/Plugins/Online/Android/AndroidFetchBackgroundDownload/Source/AndroidFetchBackgroundDownload/AndroidFetchBackgroundDownload_UPL.xml
--- a/Engine/Plugins/Online/Android/AndroidFetchBackgroundDownload/Source/AndroidFetchBackgroundDownload/AndroidFetchBackgroundDownload_UPL.xml
+++ b/Engine/Plugins/Online/Android/AndroidFetchBackgroundDownload/Source/AndroidFetchBackgroundDownload/AndroidFetchBackgroundDownload_UPL.xml
@@ -9,9 +9,13 @@
 	<buildGradleAdditions>
 		<insert>
 //AndroidFetchBackgroundDownload -- START
+	repositories {
+		mavenCentral()
+		maven { url 'https://jitpack.io' }
+	}
 	dependencies {
 		//Fetch2 library used to manage downloads at java level
-		implementation "androidx.tonyodev.fetch2:xfetch2:3.1.6"
+		implementation "com.github.tonyofrancis:fetch:3.1.6"
 		
 		//Used for parsing from JSON files in DownloadDescription.java
 		//excluding hamcrest to avoid collision with junit that also includes hamcrest
