Enable HTTP/2 for Android Fetch background downloading
- adds -allowhttp2 commandline option to enable
- http.CurlAllowHTTP2=true for libcurl support for HTTP/2
- AndroidCompat.http.CurlAllowHTTP2=true in configrules to enable for libcurl and BGDL

#android

diff --git a/Engine/Plugins/Online/Android/AndroidFetchBackgroundDownload/Source/AndroidFetchBackgroundDownload/AndroidFetchBackgroundDownload_UPL.xml b/Engine/Plugins/Online/Android/AndroidFetchBackgroundDownload/Source/AndroidFetchBackgroundDownload/AndroidFetchBackgroundDownload_UPL.xml
--- a/Engine/Plugins/Online/Android/AndroidFetchBackgroundDownload/Source/AndroidFetchBackgroundDownload/AndroidFetchBackgroundDownload_UPL.xml
+++ b/Engine/Plugins/Online/Android/AndroidFetchBackgroundDownload/Source/AndroidFetchBackgroundDownload/AndroidFetchBackgroundDownload_UPL.xml
@@ -15,7 +15,9 @@
 	}
 	dependencies {
 		//Fetch2 library used to manage downloads at java level
-		implementation "com.github.tonyofrancis:fetch:3.1.6"
+		implementation "com.github.tonyofrancis:fetch:3.1.6" 
+		api "com.squareup.okhttp3:okhttp:3.12.12"
+		api "com.squareup.okhttp3:okhttp-urlconnection:3.12.12"
 		
 		//Used for parsing from JSON files in DownloadDescription.java
 		//excluding hamcrest to avoid collision with junit that also includes hamcrest
diff --git a/Engine/Plugins/Online/Android/AndroidFetchBackgroundDownload/Source/AndroidFetchBackgroundDownload/Java/download/fetch/FetchManager.java b/Engine/Plugins/Online/Android/AndroidFetchBackgroundDownload/Source/AndroidFetchBackgroundDownload/Java/download/fetch/FetchManager.java
--- a/Engine/Plugins/Online/Android/AndroidFetchBackgroundDownload/Source/AndroidFetchBackgroundDownload/Java/download/fetch/FetchManager.java
+++ b/Engine/Plugins/Online/Android/AndroidFetchBackgroundDownload/Source/AndroidFetchBackgroundDownload/Java/download/fetch/FetchManager.java
@@ -2,6 +2,7 @@
 package com.epicgames.unreal.download.fetch;
 
 import android.content.Context;
+import android.content.SharedPreferences;
 import android.net.Uri;
 
 import com.epicgames.unreal.Logger;
@@ -31,6 +32,9 @@ import com.tonyodev.fetch2core.Func;
 import com.tonyodev.fetch2core.Func2;
 import com.tonyodev.fetch2.Status;
 
+import com.tonyodev.fetch2okhttp.OkHttpDownloader;
+import okhttp3.OkHttpClient;
+
 import androidx.annotation.Nullable;
 import androidx.annotation.NonNull;
 
@@ -455,6 +459,17 @@ public class FetchManager implements FetchDownloadProgressOwner, FetchEnqueueRes
 			FetchListener = null;
 		}
 
+		if (AllowHTTP2 && okHttpClient == null)
+		{
+			SharedPreferences preferences = context.getSharedPreferences("BackgroundDownload", context.MODE_PRIVATE);
+			AllowHTTP2 = preferences.getBoolean("AllowHTTP2", false);
+
+			if (AllowHTTP2)
+			{
+				okHttpClient = new OkHttpClient.Builder().build();
+			}
+		}
+
 		//Synchronized with StopWork to make sure we aren't creating and closing our fetch instance in a weird race condition
 		synchronized(this)
 		{
@@ -481,13 +496,28 @@ public class FetchManager implements FetchDownloadProgressOwner, FetchEnqueueRes
 			else
 			{
 				//TODO TRoss: Pull these values from the worker's getInputData
-				FetchInstance = Fetch.Impl.getInstance(new FetchConfiguration.Builder(context)
-					.setNamespace(context.getPackageName())
-					.enableLogging(true)
-					.setLogger(FetchLog)
-					.enableRetryOnNetworkGain(true)
-					.setProgressReportingInterval(200)
-					.build());
+				if (!AllowHTTP2)
+				{
+					FetchInstance = Fetch.Impl.getInstance(new FetchConfiguration.Builder(context)
+						.setNamespace(context.getPackageName())
+						.enableLogging(true)
+						.setLogger(FetchLog)
+						.enableRetryOnNetworkGain(true)
+						.setProgressReportingInterval(200)
+						.build());
+				}
+				else
+				{
+					FetchInstance = Fetch.Impl.getInstance(new FetchConfiguration.Builder(context)
+						.setNamespace(context.getPackageName())
+						.enableLogging(true)
+						.setLogger(FetchLog)
+						.enableRetryOnNetworkGain(true)
+						.setProgressReportingInterval(200)
+						.setDownloadConcurrentLimit(10)
+						.setHttpDownloader(new OkHttpDownloader(okHttpClient))
+						.build());
+				}
 			}
 			
 			if (!IsFetchInstanceValid())
@@ -1301,7 +1331,10 @@ public class FetchManager implements FetchDownloadProgressOwner, FetchEnqueueRes
 	private volatile Fetch FetchInstance = null;
 
 	private volatile FetchRequestProgressListener FetchListener = null;
-	
+
+	private volatile OkHttpClient okHttpClient = null;
+	private volatile boolean AllowHTTP2 = true;
+
 	private volatile HashMap<String, DownloadDescription> RequestedDownloads = new HashMap<String, DownloadDescription>();
 	private volatile HashMap<String, DownloadDescription> CompletedDownloads = new HashMap<String, DownloadDescription>();
 	private volatile HashMap<String, DownloadDescription> FailedDownloads = new HashMap<String, DownloadDescription>();
diff --git a/Engine/Source/Runtime/Online/HTTP/Private/Curl/CurlHttp.cpp b/Engine/Source/Runtime/Online/HTTP/Private/Curl/CurlHttp.cpp
--- a/Engine/Source/Runtime/Online/HTTP/Private/Curl/CurlHttp.cpp
+++ b/Engine/Source/Runtime/Online/HTTP/Private/Curl/CurlHttp.cpp
@@ -35,6 +35,13 @@ TAutoConsoleVariable<bool> CVarHttpCurlReadContentLengthWhenFinish(
 	ECVF_SaveForNextBoot
 );
 
+TAutoConsoleVariable<bool> CVarHttpCurlAllowHTTP2(
+	TEXT("http.CurlAllowHTTP2"),
+	false,
+	TEXT("Whether to allow HTTP/2 for curl requests"),
+	ECVF_Default
+);
+
 #if WITH_SSL
 static int SslCertVerify(int PreverifyOk, X509_STORE_CTX* Context)
 {
@@ -113,7 +120,14 @@ FCurlHttpRequest::FCurlHttpRequest()
 	// that reason we're disabling its use by default in the general purpose curl request wrapper and only
 	// allowing use of HTTP2 from other curl wrappers like the DerivedDataCache one.
 	// Note that CURL_HTTP_VERSION_1_1 was the default for libcurl version before 7.62.0
-	curl_easy_setopt(EasyHandle, CURLOPT_HTTP_VERSION, CURL_HTTP_VERSION_1_1);
+	if (CVarHttpCurlAllowHTTP2.GetValueOnAnyThread())
+	{
+		curl_easy_setopt(EasyHandle, CURLOPT_HTTP_VERSION, CURL_HTTP_VERSION_2TLS);
+	}
+	else
+	{
+		curl_easy_setopt(EasyHandle, CURLOPT_HTTP_VERSION, CURL_HTTP_VERSION_1_1);
+	}
 
 	// set certificate verification (disable to allow self-signed certificates)
 	if (FCurlHttpManager::CurlRequestOptions.bVerifyPeer)
diff --git a/Engine/Source/Runtime/Online/HTTP/Private/Curl/CurlHttpManager.cpp b/Engine/Source/Runtime/Online/HTTP/Private/Curl/CurlHttpManager.cpp
--- a/Engine/Source/Runtime/Online/HTTP/Private/Curl/CurlHttpManager.cpp
+++ b/Engine/Source/Runtime/Online/HTTP/Private/Curl/CurlHttpManager.cpp
@@ -182,6 +182,7 @@ void FCurlHttpManager::InitCurl()
 			PrintCurlFeature(CURL_VERSION_IDN);
 			PrintCurlFeature(CURL_VERSION_CONV);
 			PrintCurlFeature(CURL_VERSION_TLSAUTH_SRP);
+			PrintCurlFeature(CURL_VERSION_HTTP2); 
 #undef PrintCurlFeature
 		}
 
