Fix issue with Android BGDL crash and "Something went wrong" dialog

#android

diff --git a/Engine/Plugins/Online/Android/AndroidFetchBackgroundDownload/Source/AndroidFetchBackgroundDownload/Java/download/UEDownloadWorker.java b/Engine/Plugins/Online/Android/AndroidFetchBackgroundDownload/Source/AndroidFetchBackgroundDownload/Java/download/UEDownloadWorker.java
--- a/Engine/Plugins/Online/Android/AndroidFetchBackgroundDownload/Source/AndroidFetchBackgroundDownload/Java/download/UEDownloadWorker.java
+++ b/Engine/Plugins/Online/Android/AndroidFetchBackgroundDownload/Source/AndroidFetchBackgroundDownload/Java/download/UEDownloadWorker.java
@@ -245,7 +245,10 @@ public class UEDownloadWorker extends UEWorker implements DownloadProgressListen
 	public void CleanUp(String WorkID)
 	{
 		//Call stop work to make sure Fetch stops doing work while 
-		mFetchManager.StopWork(WorkID);
+		if (mFetchManager != null)
+		{
+			mFetchManager.StopWork(WorkID);
+		}
 		
 		//Clean up our DownloadDescriptionList file if our work is not going to re-run ever
 		if (ShouldCleanupDownloadDescriptorJSONFile())
