Enable BGDL for all Android devices again
Fix issue with HTTP/2 enable for Android from configrules

#android

diff --git a/Engine/Plugins/Online/Android/AndroidFetchBackgroundDownload/Source/AndroidFetchBackgroundDownload/Private/AndroidPlatformBackgroundHttpManager.cpp b/Engine/Plugins/Online/Android/AndroidFetchBackgroundDownload/Source/AndroidFetchBackgroundDownload/Private/AndroidPlatformBackgroundHttpManager.cpp
--- a/Engine/Plugins/Online/Android/AndroidFetchBackgroundDownload/Source/AndroidFetchBackgroundDownload/Private/AndroidPlatformBackgroundHttpManager.cpp
+++ b/Engine/Plugins/Online/Android/AndroidFetchBackgroundDownload/Source/AndroidFetchBackgroundDownload/Private/AndroidPlatformBackgroundHttpManager.cpp
@@ -662,7 +662,7 @@ const FString FAndroidPlatformBackgroundHttpManager::GetFullFileNameForDownloadD
 	if (!ensureAlwaysMsgf((AppendedFileNameInt < MAX_NUM_DOWNLOAD_DESC_FILES), TEXT("DownloadDescriptionList folder full of files! May lead to cases where we stomp expected .ini files for other workers!")))
 	{
 		static int StompNum = 0;
-		AppendedFileNameInt = (StompNum % MAX_NUM_DOWNLOAD_DESC_FILES);
+		AppendedFileNameInt = (StompNum++ % MAX_NUM_DOWNLOAD_DESC_FILES);
 	}
 	
 	return GetBaseFileNameForDownloadDescriptionListWithAppendedInt(AppendedFileNameInt);;
