Update to Fetch2 3.3.0 to fix issue with enqueuing new requests after previous batch completes

#android

diff --git a/Engine/Plugins/Online/Android/AndroidFetchBackgroundDownload/Source/AndroidFetchBackgroundDownload/AndroidFetchBackgroundDownload_UPL.xml b/Engine/Plugins/Online/Android/AndroidFetchBackgroundDownload/Source/AndroidFetchBackgroundDownload/AndroidFetchBackgroundDownload_UPL.xml
--- a/Engine/Plugins/Online/Android/AndroidFetchBackgroundDownload/Source/AndroidFetchBackgroundDownload/AndroidFetchBackgroundDownload_UPL.xml
+++ b/Engine/Plugins/Online/Android/AndroidFetchBackgroundDownload/Source/AndroidFetchBackgroundDownload/AndroidFetchBackgroundDownload_UPL.xml
@@ -15,7 +15,7 @@
 	}
 	dependencies {
 		//Fetch2 library used to manage downloads at java level
-		implementation "com.github.tonyofrancis:fetch:3.1.6" 
+		implementation "com.github.tonyofrancis:fetch:3.3.0" 
 		api "com.squareup.okhttp3:okhttp:3.12.12"
 		api "com.squareup.okhttp3:okhttp-urlconnection:3.12.12"
 		
diff --git a/Engine/Plugins/Online/Android/AndroidFetchBackgroundDownload/Source/AndroidFetchBackgroundDownload/Java/download/datastructs/DownloadDescription.java b/Engine/Plugins/Online/Android/AndroidFetchBackgroundDownload/Source/AndroidFetchBackgroundDownload/Java/download/datastructs/DownloadDescription.java
--- a/Engine/Plugins/Online/Android/AndroidFetchBackgroundDownload/Source/AndroidFetchBackgroundDownload/Java/download/datastructs/DownloadDescription.java
+++ b/Engine/Plugins/Online/Android/AndroidFetchBackgroundDownload/Source/AndroidFetchBackgroundDownload/Java/download/datastructs/DownloadDescription.java
@@ -77,16 +77,6 @@ public class DownloadDescription
 	{
 		boolean bSucceeded = true;
 		
-		//Generate our JSON output for the file
-		org.json.JSONArray DescriptionJsonArray = new org.json.JSONArray();
-		{
-			for (int DescriptionIndex = 0; DescriptionIndex < DownloadDescriptions.size(); ++DescriptionIndex)
-			{
-				DownloadDescription Description = DownloadDescriptions.get(DescriptionIndex);
-				DescriptionJsonArray.put(Description.ToJSON());
-			}
-		}
-		
 		//Attempt to write out JSONArray string to file
 		FileWriter JsonFile = null;
 		{
@@ -110,13 +100,27 @@ public class DownloadDescription
 				JsonFile = new FileWriter(NewFile, false);
 				
 				//Can actually write the json array out to the file now
-				JsonFile.write(DescriptionJsonArray.toString());
+				JsonFile.write("[");
+				for (int DescriptionIndex = 0; DescriptionIndex < DownloadDescriptions.size(); ++DescriptionIndex)
+				{
+					if (DescriptionIndex > 0) {
+						JsonFile.write(",\n");
+					}
+					DownloadDescription Description = DownloadDescriptions.get(DescriptionIndex);
+					JsonFile.write(Description.ToJSON().toString());
+				}
+				JsonFile.write("]\n");
 			}
 			catch (IOException e)
 			{
 				bSucceeded = false;
 				e.printStackTrace();
 			}
+			catch (Exception e)
+			{
+				bSucceeded = false;
+				e.printStackTrace();
+			}
 			finally
 			{
 				//even if our write fails, we still want to try and flush and close the file handle
@@ -138,7 +142,7 @@ public class DownloadDescription
 		
 		return bSucceeded;
 	}
-	
+
 	//Parses the supplied filename to try and get a JSON String output of the contents of the file.
 	public static String GetDownloadDescriptionJSONArrayFromFile(String FileName)
 	{
@@ -175,7 +179,7 @@ public class DownloadDescription
 		
 		return JSONOutputString;
 	}
-	
+
 	//Returns a new IndividualDownloadDescroption based on the JSONString. Returns null if error
 	public static DownloadDescription FromJSON(String JSONString)
 	{
