Update Fetch2 to fix hang issues with multiple enqueue groups
Add -nobgdl commandline option to force off for QA testing
Enable cellular notification handling
Fix resuming from waiting for cellular if WiFi reconnects
Handle lost/regained internet connection better
Don't show "Download in Progress" for lost internet or waiting for cellular permission notifications

#android

diff --git a/Engine/Plugins/Online/Android/AndroidFetchBackgroundDownload/Source/AndroidFetchBackgroundDownload/Java/download/UEDownloadWorker.java b/Engine/Plugins/Online/Android/AndroidFetchBackgroundDownload/Source/AndroidFetchBackgroundDownload/Java/download/UEDownloadWorker.java
--- a/Engine/Plugins/Online/Android/AndroidFetchBackgroundDownload/Source/AndroidFetchBackgroundDownload/Java/download/UEDownloadWorker.java
+++ b/Engine/Plugins/Online/Android/AndroidFetchBackgroundDownload/Source/AndroidFetchBackgroundDownload/Java/download/UEDownloadWorker.java
@@ -182,11 +182,23 @@ public class UEDownloadWorker extends UEWorker implements DownloadProgressListen
 		boolean ShouldHandleCellular = false;
 		if (NotificationDescription != null)
 		{
-			ShouldHandleCellular = NotificationDescription.ShouldHandleCellular;
+			ShouldHandleCellular = true; //NotificationDescription.ShouldHandleCellular;
 		}
 		// This should be handled by the Game thread otherwise
 		if (!bGameThreadIsActive && ShouldHandleCellular)
 		{
+			// Pause all downloads if we lose all network connectivity
+			if (bLostNetwork)
+			{
+				if (!bWaitingForNetwork)
+				{
+					// wait for network to return and pause all downloads
+					bWaitingForNetwork = true;
+					mFetchManager.PauseAllDownloads();
+				}
+				return;
+			}
+
 			// Pause all downloads if the current cellular preference does not allow cellular downloading
 			NetworkConnectivityClient.NetworkTransportType networkType = NetworkChangedManager.getInstance().networkTransportTypeCheck();
 			if (!bWaitingForCellularApproval && networkType == NetworkConnectivityClient.NetworkTransportType.CELLULAR)
@@ -196,9 +208,18 @@ public class UEDownloadWorker extends UEWorker implements DownloadProgressListen
 				boolean allowCell = (preferences.getInt("AllowCellular", 0) > 0);
 				if (!allowCell)
 				{
-					mFetchManager.PauseAllDownloads();
+					if (!bWaitingForNetwork)
+					{
+						mFetchManager.PauseAllDownloads();
+					}
 					bWaitingForCellularApproval = true;
 				}
+				else if (bWaitingForNetwork)
+				{
+					// network reconnected so resume (cellular allowed)
+					mFetchManager.ResumeAllDownloads();
+					bWaitingForNetwork = false;
+				}
 			}
 			else if (bWaitingForCellularApproval)
 			{
@@ -209,8 +230,23 @@ public class UEDownloadWorker extends UEWorker implements DownloadProgressListen
 				{
 					mFetchManager.ResumeAllDownloads();
 					bWaitingForCellularApproval = false;
+					bWaitingForNetwork = false;
+				}
+				else if (networkType == NetworkConnectivityClient.NetworkTransportType.WIFI ||
+						networkType == NetworkConnectivityClient.NetworkTransportType.ETHERNET)
+				{
+					// WiFi or Ethernet connected while waiting for cellular approval
+					mFetchManager.ResumeAllDownloads();
+					bWaitingForCellularApproval = false;
+					bWaitingForNetwork = false;
 				}
 			}
+			else if (bWaitingForNetwork)
+			{
+				// network reconnected so resume (cellular checks done above)
+				mFetchManager.ResumeAllDownloads();
+				bWaitingForNetwork = false;
+			}
 		}
 	}
 	
@@ -342,6 +378,9 @@ public class UEDownloadWorker extends UEWorker implements DownloadProgressListen
 			{
 				NotificationTextToUse = Description.ContentText.replace("%3d%%", "");
 			}
+
+			// for now don't show "Download in Progress"
+			NotificationTextToUse = "";
 		}
 		else
 		{
@@ -385,6 +424,9 @@ public class UEDownloadWorker extends UEWorker implements DownloadProgressListen
 			{
 				NotificationTextToUse = Description.ContentText.replace("%3d%%", "");
 			}
+			
+			// for now don't show "Download in Progress"
+			NotificationTextToUse = "";
 		}
 		else
 		{
@@ -674,6 +716,7 @@ public class UEDownloadWorker extends UEWorker implements DownloadProgressListen
 	private boolean bWaitingForCellularApproval = false;
 	private NetworkConnectivityClient.Listener NetworkListener = null;
 	private boolean bLostNetwork = false;
+	private boolean bWaitingForNetwork = false;
 	private boolean bForceStopped = false;
 	private DownloadQueueDescription QueueDescription = null;
 	private volatile boolean bHasEnqueueHappened = false;
