Rework network activity for Android BGDL

#android

diff --git a/Engine/Plugins/Online/Android/AndroidFetchBackgroundDownload/Config/BaseAndroidFetchBackgroundDownload.ini b/Engine/Plugins/Online/Android/AndroidFetchBackgroundDownload/Config/BaseAndroidFetchBackgroundDownload.ini
--- a/Engine/Plugins/Online/Android/AndroidFetchBackgroundDownload/Config/BaseAndroidFetchBackgroundDownload.ini
+++ b/Engine/Plugins/Online/Android/AndroidFetchBackgroundDownload/Config/BaseAndroidFetchBackgroundDownload.ini
@@ -85,6 +85,40 @@
 +DefaultNotificationText_NoInternet=("ru","Нет доступа в интернет")
 +DefaultNotificationText_NoInternet=("tr","İnternet bağlantısı yok")
 
+;Airplane Mode text
++DefaultNotificationText_AirplaneMode=("native","Airplane Mode Enabled. Download paused")
++DefaultNotificationText_AirplaneMode=("ar","تم تفعيل وضع الطائرة. تم إيقاف التنزيل مؤقتًا")
++DefaultNotificationText_AirplaneMode=("de","Flugzeugmodus aktiviert. Download pausiert")
++DefaultNotificationText_AirplaneMode=("es","Modo avión habilitado. Descarga pausada.")
++DefaultNotificationText_AirplaneMode=("es-419","Modo Avión habilitado. Descarga pausada")
++DefaultNotificationText_AirplaneMode=("fr","Mode avion activé. Téléchargement en pause")
++DefaultNotificationText_AirplaneMode=("it","Modalità aereo abilitata. Download in pausa")
++DefaultNotificationText_AirplaneMode=("ja","機内モードが有効にされました。ダウンロード一時停止中")
++DefaultNotificationText_AirplaneMode=("ko","비행기 모드가 활성화되었습니다. 다운로드 일시정지됨")
++DefaultNotificationText_AirplaneMode=("pl","Włączono tryb samolotowy. Pobieranie wstrzymane")
++DefaultNotificationText_AirplaneMode=("pt-BR","Modo Avião Habilitado. Download pausado")
++DefaultNotificationText_AirplaneMode=("ru","Авиарежим включён. Загрузка приостановлена")
++DefaultNotificationText_AirplaneMode=("tr","Uçak Modu Etkin. İndirme durduruldu")
++DefaultNotificationText_AirplaneMode=("zh-CN","飞机模式已启用。下载已暂停")
++DefaultNotificationText_AirplaneMode=("zh-Hant","已啟用飛行模式。 下載已暫停")
+
+;Data Saver Enabled Text
++DefaultNotificationText_DataSaverEnabled=("native","Data Saver Enabled. Download paused")
++DefaultNotificationText_DataSaverEnabled=("ar","تم تفعيل حافظ البيانات. تم إيقاف التنزيل مؤقتًا")
++DefaultNotificationText_DataSaverEnabled=("de","Flugzeugmodus aktiviert. Download pausiert")
++DefaultNotificationText_DataSaverEnabled=("es","Ahorro de datos habilitado. Descarga pausada.")
++DefaultNotificationText_DataSaverEnabled=("es-419","Almacenamiento de datos habilitado. Descarga pausada")
++DefaultNotificationText_DataSaverEnabled=("fr","Mode économiseur de données activé. Téléchargement en pause")
++DefaultNotificationText_DataSaverEnabled=("it","Risparmio dati abilitato. Download in pausa")
++DefaultNotificationText_DataSaverEnabled=("ja","データセーバーが有効にされました。ダウンロード一時停止中")
++DefaultNotificationText_DataSaverEnabled=("ko","데이터 절약 모드가 활성화되었습니다. 다운로드 일시정지됨")
++DefaultNotificationText_DataSaverEnabled=("pl","Włączono oszczędzanie danych. Pobieranie wstrzymane")
++DefaultNotificationText_DataSaverEnabled=("pt-BR","Proteção de Dados Habilitada. Download pausado")
++DefaultNotificationText_DataSaverEnabled=("ru","Сохранение данных включено. Загрузка приостановлена")
++DefaultNotificationText_DataSaverEnabled=("tr","Veri Tasarrufu Etkin. İndirme durduruldu")
++DefaultNotificationText_DataSaverEnabled=("zh-CN","流量节省已启用。下载已暂停")
++DefaultNotificationText_DataSaverEnabled=("zh-Hant","已啟用數據儲存。 下載已暫停")
+
 ;Waiting for cellular text
 +DefaultNotificationText_WaitingForCellular=("native","Waiting for cellular approval")
 +DefaultNotificationText_WaitingForCellular=("ar","في انتظار السماح باستخدام بيانات الهاتف")
diff --git a/Engine/Plugins/Online/Android/AndroidFetchBackgroundDownload/Source/AndroidFetchBackgroundDownload/Java/download/UEDownloadWorker.java b/Engine/Plugins/Online/Android/AndroidFetchBackgroundDownload/Source/AndroidFetchBackgroundDownload/Java/download/UEDownloadWorker.java
--- a/Engine/Plugins/Online/Android/AndroidFetchBackgroundDownload/Source/AndroidFetchBackgroundDownload/Java/download/UEDownloadWorker.java
+++ b/Engine/Plugins/Online/Android/AndroidFetchBackgroundDownload/Source/AndroidFetchBackgroundDownload/Java/download/UEDownloadWorker.java
@@ -47,6 +47,11 @@ import static com.tonyodev.fetch2.util.FetchUtils.canRetryDownload;
 
 import java.util.concurrent.TimeUnit;
 
+import android.app.Activity;
+import android.net.ConnectivityManager;
+import android.content.Context;
+import android.provider.Settings;
+
 import static android.content.Context.NOTIFICATION_SERVICE;
 
 //Helper class to manage our different work requests and callbacks
@@ -58,6 +63,14 @@ public class UEDownloadWorker extends UEWorker implements DownloadProgressListen
 		Error,
 		OutOfRetries
 	}
+
+	public enum MeteredType
+	{
+		UNMETERED,
+		METERED_DATASAVER_ON,
+		METERED_DATASAVER_ALLOWED,
+		METERED_DATASAVER_OFF
+	}
 	
 	public UEDownloadWorker(Context context, WorkerParameters params)
 	{
@@ -66,7 +79,56 @@ public class UEDownloadWorker extends UEWorker implements DownloadProgressListen
 		//Overwrite the default log to have a more specific log identifier tag
 		Log = new Logger("UE", "UEDownloadWorker");
 	}
-	
+
+	private boolean checkAirplaneMode()
+	{
+		Context context = getApplicationContext();
+		return Settings.Global.getInt(context.getContentResolver(), Settings.Global.AIRPLANE_MODE_ON, 0) != 0;
+	}
+
+	private MeteredType checkMetered()
+	{
+		if (null == connectivityManager)
+		{
+			Context context = getApplicationContext();
+			connectivityManager = (ConnectivityManager) context.getSystemService(Context.CONNECTIVITY_SERVICE);
+		}
+
+		if (connectivityManager.isActiveNetworkMetered())
+		{
+			switch (connectivityManager.getRestrictBackgroundStatus())
+			{
+				case android.net.ConnectivityManager.RESTRICT_BACKGROUND_STATUS_ENABLED:
+					// background data usage blocked for this app.
+					return MeteredType.METERED_DATASAVER_ON;
+
+				case android.net.ConnectivityManager.RESTRICT_BACKGROUND_STATUS_WHITELISTED:
+					// app allowed to bypass Data Saver, try to minimize Data
+					return MeteredType.METERED_DATASAVER_ALLOWED;
+
+				case android.net.ConnectivityManager.RESTRICT_BACKGROUND_STATUS_DISABLED:
+					// Data Saver is disabled.
+					return MeteredType.METERED_DATASAVER_OFF;
+			}
+		}
+		return MeteredType.UNMETERED;
+	}
+
+	private boolean checkCellularAllowed()
+	{
+		if (null == preferences)
+		{
+			Context context = getApplicationContext();
+			preferences = context.getSharedPreferences("CellularNetworkPreferences", context.MODE_PRIVATE);
+		}
+		boolean allowCell = (preferences.getInt("AllowCellular", 0) > 0);
+
+		// treat metered data saver allowed as permission
+		allowCell |= (meteredType == MeteredType.METERED_DATASAVER_ALLOWED);
+
+		return allowCell;
+	}
+		
 	@Override
 	public void InitWorker()
 	{
@@ -179,74 +241,81 @@ public class UEDownloadWorker extends UEWorker implements DownloadProgressListen
 	// Will pause all downloads if the network type is Cellular
 	private void NetworkTypeCheck()
 	{
-		boolean ShouldHandleCellular = false;
-		if (NotificationDescription != null)
-		{
-			ShouldHandleCellular = true; //NotificationDescription.ShouldHandleCellular;
-		}
+//		boolean ShouldHandleCellular = true;
+//		if (NotificationDescription != null)
+//		{
+//			ShouldHandleCellular = true; //NotificationDescription.ShouldHandleCellular;
+//		}
+
 		// This should be handled by the Game thread otherwise
-		if (!bGameThreadIsActive && ShouldHandleCellular)
+//		if (!bGameThreadIsActive && ShouldHandleCellular)
 		{
-			// Pause all downloads if we lose all network connectivity
-			if (bLostNetwork)
+			NetworkConnectivityClient.NetworkTransportType networkType = NetworkChangedManager.getInstance().networkTransportTypeCheck();
+			meteredType = checkMetered();
+
+			// check if network disabled
+			bAirplaneMode = checkAirplaneMode();
+			if (bAirplaneMode)
+			{
+				// allow direct ethernet connection in airplane mode
+				if (networkType == NetworkConnectivityClient.NetworkTransportType.ETHERNET)
+				{
+					bAirplaneMode = false;
+				}
+			}
+
+			// Pause all downloads if we lose all network connectivity, airplane mode, or data saver enabled
+			if (bLostNetwork || bAirplaneMode || meteredType == MeteredType.METERED_DATASAVER_ON)
 			{
 				if (!bWaitingForNetwork)
 				{
 					// wait for network to return and pause all downloads
 					bWaitingForNetwork = true;
+					bWaitingForCellularApproval = true;
 					mFetchManager.PauseAllDownloads();
 				}
 				return;
 			}
 
-			// Pause all downloads if the current cellular preference does not allow cellular downloading
-			NetworkConnectivityClient.NetworkTransportType networkType = NetworkChangedManager.getInstance().networkTransportTypeCheck();
-			if (!bWaitingForCellularApproval && networkType == NetworkConnectivityClient.NetworkTransportType.CELLULAR)
+			// Need permission for using cellular
+			if (networkType == NetworkConnectivityClient.NetworkTransportType.CELLULAR)
 			{
-				Context context = getApplicationContext();
-				SharedPreferences preferences = context.getSharedPreferences("CellularNetworkPreferences", context.MODE_PRIVATE);
-				boolean allowCell = (preferences.getInt("AllowCellular", 0) > 0);
-				if (!allowCell)
+				// check if we have permission granted
+				boolean allowCell = checkCellularAllowed();
+
+				// if not already waiting, check if we have permission
+				if (!bWaitingForCellularApproval)
 				{
-					if (!bWaitingForNetwork)
+					if (!allowCell)
 					{
+						// no permission so pause and wait for it
 						mFetchManager.PauseAllDownloads();
+						bWaitingForCellularApproval = true;
 					}
-					bWaitingForCellularApproval = true;
 				}
-				else if (bWaitingForNetwork)
+				else
 				{
-					// network reconnected so resume (cellular allowed)
-					mFetchManager.ResumeAllDownloads();
-					bWaitingForNetwork = false;
+					// see if we got permission
+					if (allowCell)
+					{
+						// permission granted so resume download
+						mFetchManager.ResumeAllDownloads();
+						bWaitingForNetwork = false;
+						bWaitingForCellularApproval = false;
+					}
 				}
 			}
-			else if (bWaitingForCellularApproval)
+			else
 			{
-				Context context = getApplicationContext();
-				SharedPreferences preferences = context.getSharedPreferences("CellularNetworkPreferences", context.MODE_PRIVATE);
-				boolean allowCell = (preferences.getInt("AllowCellular", 0) > 0);
-				if (allowCell)
+				// have a network connection that is not cellular
+				if (bWaitingForNetwork || bWaitingForCellularApproval)
 				{
+					// resume any downloads
 					mFetchManager.ResumeAllDownloads();
-					bWaitingForCellularApproval = false;
 					bWaitingForNetwork = false;
-				}
-				else if (networkType == NetworkConnectivityClient.NetworkTransportType.WIFI ||
-						networkType == NetworkConnectivityClient.NetworkTransportType.ETHERNET)
-				{
-					// WiFi or Ethernet connected while waiting for cellular approval
-					mFetchManager.ResumeAllDownloads();
 					bWaitingForCellularApproval = false;
-					bWaitingForNetwork = false;
 				}
 			}
-			else if (bWaitingForNetwork)
-			{
-				// network reconnected so resume (cellular checks done above)
-				mFetchManager.ResumeAllDownloads();
-				bWaitingForNetwork = false;
-			}
 		}
 	}
 	
@@ -344,19 +413,27 @@ public class UEDownloadWorker extends UEWorker implements DownloadProgressListen
 		}
 		
 		Notification notification = null;
-		if (bWaitingForCellularApproval)
+		if (bAirplaneMode)
 		{
-			notification = CreateCellularWaitNotification(context, Description, pendingNotificationIntent);
+			notification = CreateAirplaneModeNotification(context, Description, pendingNotificationIntent);
 		}
-		else if(!bLostNetwork)
+		else if (meteredType == MeteredType.METERED_DATASAVER_ON)
 		{
-			notification = CreateDownloadProgressNotification(context, Description, pendingNotificationIntent);
+			notification = CreateDataSaverEnabledNotification(context, Description, pendingNotificationIntent);
 		}
-		else
+		else if (bLostNetwork)
 		{
 			notification = CreateNoInternetDownloadNotification(context, Description, pendingNotificationIntent);
 		}
-		
+		else if (bWaitingForCellularApproval && !checkCellularAllowed())
+		{
+			notification = CreateCellularWaitNotification(context, Description, pendingNotificationIntent);
+		}
+		else
+		{
+			notification = CreateDownloadProgressNotification(context, Description, pendingNotificationIntent);
+		}
+
 		return new ForegroundInfo(Description.NotificationID,notification);
 	}
 
@@ -390,7 +467,7 @@ public class UEDownloadWorker extends UEWorker implements DownloadProgressListen
 		int CurrentProgress = bShowPercentage ? Description.CurrentProgress : 0;
 		boolean Indeterminate = bShowPercentage ? Description.Indeterminate : true;
 
-		Notification notification = new NotificationCompat.Builder(context, Description.NotificationChannelID)
+		NotificationCompat.Builder builder = new NotificationCompat.Builder(context, Description.NotificationChannelID)
 			.setContentTitle(Description.NoInternetAvailable)
 			.setTicker(Description.TitleText)
 			.setContentText(NotificationTextToUse)
@@ -399,14 +476,15 @@ public class UEDownloadWorker extends UEWorker implements DownloadProgressListen
 			.setOngoing(true)
 			.setOnlyAlertOnce (true)
 			.setSmallIcon(Description.SmallIconResourceID)
-			.addAction(Description.CancelIconResourceID, Description.CancelText, CancelIntent)
-			.setNotificationSilent()
-			.build();
-
-		return notification;
+			.setNotificationSilent();
+		if (!bGameThreadIsActive)
+		{
+			builder.addAction(Description.CancelIconResourceID, Description.CancelText, CancelIntent);
+		}
+		return builder.build();
 	}
 
-	public Notification CreateCellularWaitNotification(Context context,
+	public Notification CreateAirplaneModeNotification(Context context,
 		DownloadNotificationDescription Description,
 		PendingIntent pendingNotificationIntent)
 	{
@@ -424,7 +502,7 @@ public class UEDownloadWorker extends UEWorker implements DownloadProgressListen
 			{
 				NotificationTextToUse = Description.ContentText.replace("%3d%%", "");
 			}
-			
+
 			// for now don't show "Download in Progress"
 			NotificationTextToUse = "";
 		}
@@ -432,47 +510,122 @@ public class UEDownloadWorker extends UEWorker implements DownloadProgressListen
 		{
 			NotificationTextToUse = Description.ContentCompleteText;
 		}
-		PendingIntent cellularNotificationIntent = null;
+
+		int CurrentProgress = bShowPercentage ? Description.CurrentProgress : 0;
+		boolean Indeterminate = bShowPercentage ? Description.Indeterminate : true;
+
+		NotificationCompat.Builder builder = new NotificationCompat.Builder(context, Description.NotificationChannelID)
+			.setContentTitle(Description.AirplaneModeText)
+			.setTicker(Description.TitleText)
+			.setContentText(NotificationTextToUse)
+			.setContentIntent(pendingNotificationIntent)
+			.setProgress(Description.MAX_PROGRESS, CurrentProgress, Indeterminate)
+			.setOngoing(true)
+			.setOnlyAlertOnce (true)
+			.setSmallIcon(Description.SmallIconResourceID)
+			.setNotificationSilent();
+		if (!bGameThreadIsActive)
 		{
-			cellularNotificationIntent = PendingIntent.getBroadcast(context, Description.NotificationID, ApproveIntent, PendingIntent.FLAG_IMMUTABLE);
+			builder.addAction(Description.CancelIconResourceID, Description.CancelText, CancelIntent);
+		}
+		return builder.build();
+	}
+
+	public Notification CreateDataSaverEnabledNotification(Context context,
+		DownloadNotificationDescription Description,
+		PendingIntent pendingNotificationIntent)
+	{
+		//Setup Notification Text Values (ContentText and ContentInfo)
+		boolean bShowPercentage = ShouldShowPercentage(context);
+		boolean bIsComplete = (Description.CurrentProgress >= Description.MAX_PROGRESS);
+		String NotificationTextToUse = null;
+		if (!bIsComplete)
+		{
+			if (bShowPercentage)
+			{
+				NotificationTextToUse = String.format(Description.ContentText, Description.CurrentProgress);
+			}
+			else
+			{
+				NotificationTextToUse = Description.ContentText.replace("%3d%%", "");
+			}
+
+			// for now don't show "Download in Progress"
+			NotificationTextToUse = "";
+		}
+		else
+		{
+			NotificationTextToUse = Description.ContentCompleteText;
 		}
 
 		int CurrentProgress = bShowPercentage ? Description.CurrentProgress : 0;
 		boolean Indeterminate = bShowPercentage ? Description.Indeterminate : true;
 
-		Notification notification;
-		if (bGameThreadIsActive)
-		{
-			notification = new NotificationCompat.Builder(context, Description.NotificationChannelID)
-				.setContentTitle(Description.WaitingForCellularText)
-				.setTicker(Description.WaitingForCellularText)
-				.setContentText(NotificationTextToUse)
-				.setContentIntent(pendingNotificationIntent)
-				.setProgress(Description.MAX_PROGRESS, CurrentProgress, Indeterminate)
-				.setOngoing(true)
-				.setOnlyAlertOnce (true)
-				.setSmallIcon(Description.SmallIconResourceID)
-				.setNotificationSilent()
-				.build();
+		NotificationCompat.Builder builder = new NotificationCompat.Builder(context, Description.NotificationChannelID)
+			.setContentTitle(Description.DataSaverEnabledText)
+			.setTicker(Description.TitleText)
+			.setContentText(NotificationTextToUse)
+			.setContentIntent(pendingNotificationIntent)
+			.setProgress(Description.MAX_PROGRESS, CurrentProgress, Indeterminate)
+			.setOngoing(true)
+			.setOnlyAlertOnce (true)
+			.setSmallIcon(Description.SmallIconResourceID)
+			.setNotificationSilent();
+		if (!bGameThreadIsActive)
+		{
+			builder.addAction(Description.CancelIconResourceID, Description.CancelText, CancelIntent);
+		}
+		return builder.build();
+	}
+
+	public Notification CreateCellularWaitNotification(Context context,
+		DownloadNotificationDescription Description,
+		PendingIntent pendingNotificationIntent)
+	{
+		//Setup Notification Text Values (ContentText and ContentInfo)
+		boolean bShowPercentage = ShouldShowPercentage(context);
+		boolean bIsComplete = (Description.CurrentProgress >= Description.MAX_PROGRESS);
+		String NotificationTextToUse = null;
+		if (!bIsComplete)
+		{
+			if (bShowPercentage)
+			{
+				NotificationTextToUse = String.format(Description.ContentText, Description.CurrentProgress);
+			}
+			else
+			{
+				NotificationTextToUse = Description.ContentText.replace("%3d%%", "");
+			}
+			
+			// for now don't show "Download in Progress"
+			NotificationTextToUse = "";
 		}
 		else
 		{
-			notification = new NotificationCompat.Builder(context, Description.NotificationChannelID)
-				.setContentTitle(Description.WaitingForCellularText)
-				.setTicker(Description.WaitingForCellularText)
-				.setContentText(NotificationTextToUse)
-				.setContentIntent(pendingNotificationIntent)
-				.setProgress(Description.MAX_PROGRESS, CurrentProgress, Indeterminate)
-				.setOngoing(true)
-				.setOnlyAlertOnce (true)
-				.setSmallIcon(Description.SmallIconResourceID)
-				.addAction(Description.CancelIconResourceID, Description.ApproveText, cellularNotificationIntent)
-				.addAction(Description.CancelIconResourceID, Description.CancelText, CancelIntent)
-				.setNotificationSilent()
-				.build();
+			NotificationTextToUse = Description.ContentCompleteText;
 		}
-		
-		return notification;
+
+		int CurrentProgress = bShowPercentage ? Description.CurrentProgress : 0;
+		boolean Indeterminate = bShowPercentage ? Description.Indeterminate : true;
+
+		NotificationCompat.Builder builder = new NotificationCompat.Builder(context, Description.NotificationChannelID)
+			.setContentTitle(Description.WaitingForCellularText)
+			.setTicker(Description.WaitingForCellularText)
+			.setContentText(NotificationTextToUse)
+			.setContentIntent(pendingNotificationIntent)
+			.setProgress(Description.MAX_PROGRESS, CurrentProgress, Indeterminate)
+			.setOngoing(true)
+			.setOnlyAlertOnce (true)
+			.setSmallIcon(Description.SmallIconResourceID)
+			.setNotificationSilent();
+		if (!bGameThreadIsActive)
+		{
+			PendingIntent cellularNotificationIntent = null;
+			cellularNotificationIntent = PendingIntent.getBroadcast(context, Description.NotificationID, ApproveIntent, PendingIntent.FLAG_IMMUTABLE);
+			builder.addAction(Description.CancelIconResourceID, Description.ApproveText, cellularNotificationIntent);
+			builder.addAction(Description.CancelIconResourceID, Description.CancelText, CancelIntent);
+		}
+		return builder.build();
 	}
 
 	public Notification CreateDownloadProgressNotification(Context context,
@@ -673,7 +826,7 @@ public class UEDownloadWorker extends UEWorker implements DownloadProgressListen
 		boolean ShouldHandleCellular = false;
 		if (NotificationDescription != null)
 		{
-			ShouldHandleCellular = NotificationDescription.ShouldHandleCellular;
+			ShouldHandleCellular = true; // NotificationDescription.ShouldHandleCellular;
 		}
 		if (ShouldHandleCellular)
 		{
@@ -713,9 +866,14 @@ public class UEDownloadWorker extends UEWorker implements DownloadProgressListen
 	public native void nativeAndroidBackgroundDownloadOnAllComplete(boolean bDidAllRequestsSucceed);
 	public native void nativeAndroidBackgroundDownloadOnTick();
 	
+	private ConnectivityManager connectivityManager = null;
+	SharedPreferences preferences = null;
+
 	private boolean bWaitingForCellularApproval = false;
 	private NetworkConnectivityClient.Listener NetworkListener = null;
 	private boolean bLostNetwork = false;
+	private boolean bAirplaneMode = false;
+	private MeteredType meteredType = MeteredType.UNMETERED;
 	private boolean bWaitingForNetwork = false;
 	private boolean bForceStopped = false;
 	private DownloadQueueDescription QueueDescription = null;
diff --git a/Engine/Plugins/Online/Android/AndroidFetchBackgroundDownload/Source/AndroidFetchBackgroundDownload/Java/download/datastructs/DownloadNotificationDescription.java b/Engine/Plugins/Online/Android/AndroidFetchBackgroundDownload/Source/AndroidFetchBackgroundDownload/Java/download/datastructs/DownloadNotificationDescription.java
--- a/Engine/Plugins/Online/Android/AndroidFetchBackgroundDownload/Source/AndroidFetchBackgroundDownload/Java/download/datastructs/DownloadNotificationDescription.java
+++ b/Engine/Plugins/Online/Android/AndroidFetchBackgroundDownload/Source/AndroidFetchBackgroundDownload/Java/download/datastructs/DownloadNotificationDescription.java
@@ -95,6 +95,18 @@ public class DownloadNotificationDescription
 				NoInternetAvailable = "No Internet Available";
 			}
 
+			AirplaneModeText = data.getString(DownloadWorkerParameterKeys.NOTIFICATION_CONTENT_AIRPLANE_MODE_TEXT_KEY);
+			if (null == AirplaneModeText)
+			{
+				AirplaneModeText = "Airplane Mode Enabled";
+			}
+			
+			DataSaverEnabledText = data.getString(DownloadWorkerParameterKeys.NOTIFICATION_CONTENT_DATA_SAVER_ENABLED_TEXT_KEY);
+			if (null == DataSaverEnabledText)
+			{
+				DataSaverEnabledText = "Data Saver Enabled";
+			}
+
 			WaitingForCellularText = data.getString(DownloadWorkerParameterKeys.NOTIFICATION_CONTENT_WAITING_FOR_CELLULAR_TEXT_KEY);
 			if (null == WaitingForCellularText)
 			{
@@ -252,6 +264,8 @@ public class DownloadNotificationDescription
 	public String ContentCompleteText = null;
 	public String CancelText = null;
 	public String NoInternetAvailable = null;
+	public String AirplaneModeText = null;
+	public String DataSaverEnabledText = null;
 	public String WaitingForCellularText = null;
 	public String ApproveText = null;
 	
diff --git a/Engine/Plugins/Online/Android/AndroidFetchBackgroundDownload/Source/AndroidFetchBackgroundDownload/Java/download/datastructs/DownloadWorkerParameterKeys.java b/Engine/Plugins/Online/Android/AndroidFetchBackgroundDownload/Source/AndroidFetchBackgroundDownload/Java/download/datastructs/DownloadWorkerParameterKeys.java
--- a/Engine/Plugins/Online/Android/AndroidFetchBackgroundDownload/Source/AndroidFetchBackgroundDownload/Java/download/datastructs/DownloadWorkerParameterKeys.java
+++ b/Engine/Plugins/Online/Android/AndroidFetchBackgroundDownload/Source/AndroidFetchBackgroundDownload/Java/download/datastructs/DownloadWorkerParameterKeys.java
@@ -36,6 +36,8 @@ public class DownloadWorkerParameterKeys
 	public final static String NOTIFICATION_CONTENT_TEXT_KEY = "NotificationContentText";
 	public final static String NOTIFICATION_CONTENT_CANCEL_DOWNLOAD_TEXT_KEY = "NotificationContentCancelDownloadText";
 	public final static String NOTIFICATION_CONTENT_NO_INTERNET_TEXT_KEY = "NotificationContentNoInternetDownloadText";
+	public final static String NOTIFICATION_CONTENT_AIRPLANE_MODE_TEXT_KEY = "NotificationContentAirplaneModeText";
+	public final static String NOTIFICATION_CONTENT_DATA_SAVER_ENABLED_TEXT_KEY = "NotificationContentDataSaverEnabledText";
 	public final static String NOTIFICATION_CONTENT_COMPLETE_TEXT_KEY = "NotificationContentCompleteText";
 	public final static String NOTIFICATION_CONTENT_WAITING_FOR_CELLULAR_TEXT_KEY = "NotificationContentWaitingForCellularText";
 	public final static String NOTIFICATION_CONTENT_APPROVE_TEXT_KEY = "NotificationContentApproveText";
diff --git a/Engine/Plugins/Online/Android/AndroidFetchBackgroundDownload/Source/AndroidFetchBackgroundDownload/Java/download/fetch/FetchManager.java b/Engine/Plugins/Online/Android/AndroidFetchBackgroundDownload/Source/AndroidFetchBackgroundDownload/Java/download/fetch/FetchManager.java
--- a/Engine/Plugins/Online/Android/AndroidFetchBackgroundDownload/Source/AndroidFetchBackgroundDownload/Java/download/fetch/FetchManager.java
+++ b/Engine/Plugins/Online/Android/AndroidFetchBackgroundDownload/Source/AndroidFetchBackgroundDownload/Java/download/fetch/FetchManager.java
@@ -282,8 +282,11 @@ public class FetchManager implements FetchDownloadProgressOwner, FetchEnqueueRes
 				return;
 			}
 			
-			MatchedDesc.bIsPaused = true;
-			FetchInstance.pause(MatchedDesc.CachedFetchID);
+			if (!MatchedDesc.bIsPaused)
+			{
+				MatchedDesc.bIsPaused = true;
+				FetchInstance.pause(MatchedDesc.CachedFetchID);
+			}
 		}
 	}
 
@@ -304,8 +307,11 @@ public class FetchManager implements FetchDownloadProgressOwner, FetchEnqueueRes
 				return;
 			}
 			
-			MatchedDesc.bIsPaused = false;
-			FetchInstance.resume(MatchedDesc.CachedFetchID);
+			if (MatchedDesc.bIsPaused)
+			{
+				MatchedDesc.bIsPaused = false;
+				FetchInstance.resume(MatchedDesc.CachedFetchID);
+			}
 		}
 	}
 
diff --git a/Engine/Plugins/Online/Android/AndroidFetchBackgroundDownload/Source/AndroidFetchBackgroundDownload/Private/AndroidPlatformBackgroundHttpManager.cpp b/Engine/Plugins/Online/Android/AndroidFetchBackgroundDownload/Source/AndroidFetchBackgroundDownload/Private/AndroidPlatformBackgroundHttpManager.cpp
--- a/Engine/Plugins/Online/Android/AndroidFetchBackgroundDownload/Source/AndroidFetchBackgroundDownload/Private/AndroidPlatformBackgroundHttpManager.cpp
+++ b/Engine/Plugins/Online/Android/AndroidFetchBackgroundDownload/Source/AndroidFetchBackgroundDownload/Private/AndroidPlatformBackgroundHttpManager.cpp
@@ -94,6 +94,8 @@ const FString FAndroidNativeDownloadWorkerParameterKeys::NOTIFICATION_CONTENT_TI
 const FString FAndroidNativeDownloadWorkerParameterKeys::NOTIFICATION_CONTENT_TEXT_KEY = TEXT("NotificationContentText");
 const FString FAndroidNativeDownloadWorkerParameterKeys::NOTIFICATION_CONTENT_CANCEL_DOWNLOAD_TEXT_KEY = TEXT("NotificationContentCancelDownloadText");
 const FString FAndroidNativeDownloadWorkerParameterKeys::NOTIFICATION_CONTENT_NO_INTERNET_TEXT_KEY = TEXT("NotificationContentNoInternetDownloadText");
+const FString FAndroidNativeDownloadWorkerParameterKeys::NOTIFICATION_CONTENT_AIRPLANE_MODE_TEXT_KEY = TEXT("NotificationContentAirmodeModeText");
+const FString FAndroidNativeDownloadWorkerParameterKeys::NOTIFICATION_CONTENT_DATA_SAVER_ENABLED_TEXT_KEY = TEXT("NotificationContentDataSaverEnabledText");
 const FString FAndroidNativeDownloadWorkerParameterKeys::NOTIFICATION_CONTENT_COMPLETE_TEXT_KEY = TEXT("NotificationContentCompleteText");
 const FString FAndroidNativeDownloadWorkerParameterKeys::NOTIFICATION_CONTENT_WAITING_FOR_CELLULAR_TEXT_KEY = TEXT("NotificationContentWaitingForCellularText");
 const FString FAndroidNativeDownloadWorkerParameterKeys::NOTIFICATION_CONTENT_APPROVE_TEXT_KEY = TEXT("NotificationContentApproveText");
@@ -361,6 +363,8 @@ void FAndroidPlatformBackgroundHttpManager::ActivatePendingRequests()
 				WorkParams.AddDataToWorkerParameters(FAndroidNativeDownloadWorkerParameterKeys::NOTIFICATION_CONTENT_COMPLETE_TEXT_KEY, AndroidBackgroundHTTPManagerDefaultLocalizedText.DefaultNotificationText_Complete.GetText());
 				WorkParams.AddDataToWorkerParameters(FAndroidNativeDownloadWorkerParameterKeys::NOTIFICATION_CONTENT_CANCEL_DOWNLOAD_TEXT_KEY, AndroidBackgroundHTTPManagerDefaultLocalizedText.DefaultNotificationText_Cancel.GetText());
 				WorkParams.AddDataToWorkerParameters(FAndroidNativeDownloadWorkerParameterKeys::NOTIFICATION_CONTENT_NO_INTERNET_TEXT_KEY, AndroidBackgroundHTTPManagerDefaultLocalizedText.DefaultNotificationText_NoInternet.GetText());
+				WorkParams.AddDataToWorkerParameters(FAndroidNativeDownloadWorkerParameterKeys::NOTIFICATION_CONTENT_AIRPLANE_MODE_TEXT_KEY, AndroidBackgroundHTTPManagerDefaultLocalizedText.DefaultNotificationText_AirplaneMode.GetText());
+				WorkParams.AddDataToWorkerParameters(FAndroidNativeDownloadWorkerParameterKeys::NOTIFICATION_CONTENT_DATA_SAVER_ENABLED_TEXT_KEY, AndroidBackgroundHTTPManagerDefaultLocalizedText.DefaultNotificationText_DataSaverEnabled.GetText());
 				WorkParams.AddDataToWorkerParameters(FAndroidNativeDownloadWorkerParameterKeys::NOTIFICATION_CONTENT_WAITING_FOR_CELLULAR_TEXT_KEY, AndroidBackgroundHTTPManagerDefaultLocalizedText.DefaultNotificationText_WaitingForCellular.GetText());
 				WorkParams.AddDataToWorkerParameters(FAndroidNativeDownloadWorkerParameterKeys::NOTIFICATION_CONTENT_APPROVE_TEXT_KEY, AndroidBackgroundHTTPManagerDefaultLocalizedText.DefaultNotificationText_Approve.GetText());
 
@@ -1056,6 +1060,8 @@ FAndroidPlatformBackgroundHttpManager::FAndroidBackgroundHTTPManagerDefaultLocal
 	, DefaultNotificationText_Complete()
 	, DefaultNotificationText_Cancel()
 	, DefaultNotificationText_NoInternet()
+	, DefaultNotificationText_AirplaneMode()
+	, DefaultNotificationText_DataSaverEnabled()
 	, DefaultNotificationText_WaitingForCellular()
 	, DefaultNotificationText_Approve()
 {
@@ -1105,6 +1111,20 @@ void FAndroidPlatformBackgroundHttpManager::FAndroidBackgroundHTTPManagerDefault
 				ParsePolyglotTextItem(DefaultNotificationText_NoInternet, TEXT("Notification.NoInternetText"), NoInternetTextStrings);
 			}
 
+			//DefaultNotificationText_AirplaneMode
+			{
+				TArray<FString> AirplaneModeTextStrings;
+				Config->GetArray(*TextConfigSection, TEXT("DefaultNotificationText_AirplaneMode"), AirplaneModeTextStrings);
+				ParsePolyglotTextItem(DefaultNotificationText_AirplaneMode, TEXT("Notification.AirplaneModeText"), AirplaneModeTextStrings);
+			}
+
+			//DefaultNotificationText_DataSaverEnabled
+			{
+				TArray<FString> DataSaverEnabledTextStrings;
+				Config->GetArray(*TextConfigSection, TEXT("DefaultNotificationText_DataSaverEnabled"), DataSaverEnabledTextStrings);
+				ParsePolyglotTextItem(DefaultNotificationText_DataSaverEnabled, TEXT("Notification.DataSaverEnabledText"), DataSaverEnabledTextStrings);
+			}
+
 			//DefaultNotificationText_WaitingForCellular
 			{
 				TArray<FString> WaitingForCellularTextStrings;
diff --git a/Engine/Plugins/Online/Android/AndroidFetchBackgroundDownload/Source/AndroidFetchBackgroundDownload/Public/AndroidPlatformBackgroundHttpManager.h b/Engine/Plugins/Online/Android/AndroidFetchBackgroundDownload/Source/AndroidFetchBackgroundDownload/Public/AndroidPlatformBackgroundHttpManager.h
--- a/Engine/Plugins/Online/Android/AndroidFetchBackgroundDownload/Source/AndroidFetchBackgroundDownload/Public/AndroidPlatformBackgroundHttpManager.h
+++ b/Engine/Plugins/Online/Android/AndroidFetchBackgroundDownload/Source/AndroidFetchBackgroundDownload/Public/AndroidPlatformBackgroundHttpManager.h
@@ -158,6 +158,8 @@ protected:
 		FPolyglotTextData DefaultNotificationText_Complete;
 		FPolyglotTextData DefaultNotificationText_Cancel;
 		FPolyglotTextData DefaultNotificationText_NoInternet;
+		FPolyglotTextData DefaultNotificationText_AirplaneMode;
+		FPolyglotTextData DefaultNotificationText_DataSaverEnabled;
 		FPolyglotTextData DefaultNotificationText_WaitingForCellular;
 		FPolyglotTextData DefaultNotificationText_Approve;
 
@@ -229,6 +231,8 @@ public:
 	static const FString NOTIFICATION_CONTENT_TEXT_KEY;
 	static const FString NOTIFICATION_CONTENT_CANCEL_DOWNLOAD_TEXT_KEY;
 	static const FString NOTIFICATION_CONTENT_NO_INTERNET_TEXT_KEY;
+	static const FString NOTIFICATION_CONTENT_AIRPLANE_MODE_TEXT_KEY;
+	static const FString NOTIFICATION_CONTENT_DATA_SAVER_ENABLED_TEXT_KEY;
 	static const FString NOTIFICATION_CONTENT_COMPLETE_TEXT_KEY;
 	static const FString NOTIFICATION_CONTENT_WAITING_FOR_CELLULAR_TEXT_KEY;
 	static const FString NOTIFICATION_CONTENT_APPROVE_TEXT_KEY;
