Implement API needed for cellular data opt-in dialogue

diff --git a/Engine/Plugins/Online/Android/AndroidFetchBackgroundDownload/Source/AndroidFetchBackgroundDownload/Private/AndroidPlatformBackgroundHttpManager.cpp b/Engine/Plugins/Online/Android/AndroidFetchBackgroundDownload/Source/AndroidFetchBackgroundDownload/Private/AndroidPlatformBackgroundHttpManager.cpp
--- a/Engine/Plugins/Online/Android/AndroidFetchBackgroundDownload/Source/AndroidFetchBackgroundDownload/Private/AndroidPlatformBackgroundHttpManager.cpp
+++ b/Engine/Plugins/Online/Android/AndroidFetchBackgroundDownload/Source/AndroidFetchBackgroundDownload/Private/AndroidPlatformBackgroundHttpManager.cpp
@@ -178,6 +178,9 @@ void FAndroidPlatformBackgroundHttpManager::Initialize()
 	//and now that we are foregrounded we should wait to do that work until requested again.
 	FUEWorkManagerNativeWrapper::CancelBackgroundWork(BackgroundHTTPWorkID);
 
+	//Update cellular preference in case if download was allowed to go over cellular from notification.
+	CellularPreference = FPlatformMisc::GetCellularPreference();
+
 	FAndroidBackgroundDownloadDelegates::AndroidBackgroundDownload_OnWorkerStart.BindSP(this, &FAndroidPlatformBackgroundHttpManager::Java_OnWorkerStart);
 	FAndroidBackgroundDownloadDelegates::AndroidBackgroundDownload_OnWorkerStop.BindSP(this, &FAndroidPlatformBackgroundHttpManager::Java_OnWorkerStop);
 	FAndroidBackgroundDownloadDelegates::AndroidBackgroundDownload_OnProgress.BindSP(this, &FAndroidPlatformBackgroundHttpManager::Java_OnDownloadProgress);
@@ -185,7 +188,7 @@ void FAndroidPlatformBackgroundHttpManager::Initialize()
 	FAndroidBackgroundDownloadDelegates::AndroidBackgroundDownload_OnMetrics.BindSP(this, &FAndroidPlatformBackgroundHttpManager::Java_OnDownloadMetrics);
 	FAndroidBackgroundDownloadDelegates::AndroidBackgroundDownload_OnAllComplete.BindSP(this, &FAndroidPlatformBackgroundHttpManager::Java_OnAllDownloadsComplete);
 	FAndroidBackgroundDownloadDelegates::AndroidBackgroundDownload_OnTickWorkerThread.BindSP(this, &FAndroidPlatformBackgroundHttpManager::Java_OnTick);
-
+ 
 	AndroidBackgroundHTTPManagerDefaultLocalizedText.InitFromIniSettings("AndroidFetchBackgroundDownload");
 	FBackgroundHttpManagerImpl::Initialize();
 }
@@ -211,6 +214,12 @@ bool FAndroidPlatformBackgroundHttpManager::Tick(float DeltaTime)
 	return true;
 }
 
+void FAndroidPlatformBackgroundHttpManager::SetCellularPreference(int32 Value)
+{
+	CellularPreference = Value;
+	FPlatformMisc::SetCellularPreference(Value);
+}
+
 void FAndroidPlatformBackgroundHttpManager::ActivatePendingRequests()
 {
 	UE_LOG(LogBackgroundHttpManager, VeryVerbose, TEXT("ActivatePendingRequests called"));
@@ -242,6 +251,9 @@ void FAndroidPlatformBackgroundHttpManager::ActivatePendingRequests()
 		UE_CLOG(bHasPendingRequests, LogBackgroundHttpManager, Display, TEXT("Found PendingRequests to activate"));
 		UE_CLOG(bForceActiveRequestsToRequeue, LogBackgroundHttpManager, Display, TEXT("Forcing Found ActiveRequests to requeue"));
 
+		// UEDownloadWorker.java actively reads and _removes_ this setting, so setting it every time we schedule a request.
+		FPlatformMisc::SetCellularPreference(CellularPreference);
+
 		const bool bIsOptional = false;
 
 		JNIEnv* Env = FAndroidApplication::GetJavaEnv();
diff --git a/Engine/Plugins/Online/Android/AndroidFetchBackgroundDownload/Source/AndroidFetchBackgroundDownload/Public/AndroidPlatformBackgroundHttpManager.h b/Engine/Plugins/Online/Android/AndroidFetchBackgroundDownload/Source/AndroidFetchBackgroundDownload/Public/AndroidPlatformBackgroundHttpManager.h
--- a/Engine/Plugins/Online/Android/AndroidFetchBackgroundDownload/Source/AndroidFetchBackgroundDownload/Public/AndroidPlatformBackgroundHttpManager.h
+++ b/Engine/Plugins/Online/Android/AndroidFetchBackgroundDownload/Source/AndroidFetchBackgroundDownload/Public/AndroidPlatformBackgroundHttpManager.h
@@ -37,6 +37,7 @@ public:
 public:
 	virtual void Initialize() override;
 	virtual void Shutdown() override;
+	virtual void SetCellularPreference(int32 Value) override;
 	virtual bool IsGenericImplementation() const override { return false; }
 	bool Tick(float DeltaTime);
 	
@@ -128,6 +129,7 @@ protected:
 	void OnNetworkConnectionChanged(ENetworkConnectionType ConnectionType);
 	FDelegateHandle OnNetworkConnectionChangedDelegateHandle;
 
+	int32 CellularPreference = 0;
 protected:
 	//Used to determine the state of AndroidBackgroundHTTP in our configrules
 	enum class EAndroidBackgroundDownloadConfigRulesSetting : uint8
diff --git a/Engine/Plugins/Runtime/GameFeatures/Source/GameFeatures/Private/GameFeatureOptionalContentInstaller.cpp b/Engine/Plugins/Runtime/GameFeatures/Source/GameFeatures/Private/GameFeatureOptionalContentInstaller.cpp
--- a/Engine/Plugins/Runtime/GameFeatures/Source/GameFeatures/Private/GameFeatureOptionalContentInstaller.cpp
+++ b/Engine/Plugins/Runtime/GameFeatures/Source/GameFeatures/Private/GameFeatureOptionalContentInstaller.cpp
@@ -77,6 +77,7 @@ void UGameFeatureOptionalContentInstaller::EnableCellularDownloading(bool bEnabl
 	}
 
 	bAllowCellDownload = bEnable;
+	BundleManager->SetCellularPreference(bAllowCellDownload ? 1 : 0);
 
 	// Update flags on active requests
 	for ( TPair<FString, FGFPInstall>& Pair : ActiveGFPInstalls)
diff --git a/Engine/Source/Runtime/InstallBundleManager/Public/InstallBundleManagerInterface.h b/Engine/Source/Runtime/InstallBundleManager/Public/InstallBundleManagerInterface.h
--- a/Engine/Source/Runtime/InstallBundleManager/Public/InstallBundleManagerInterface.h
+++ b/Engine/Source/Runtime/InstallBundleManager/Public/InstallBundleManagerInterface.h
@@ -191,6 +191,14 @@ public:
 	INSTALLBUNDLEMANAGER_API void UpdateContentRequestFlags(FName BundleName, EInstallBundleRequestFlags AddFlags, EInstallBundleRequestFlags RemoveFlags);
 	virtual void UpdateContentRequestFlags(TArrayView<const FName> BundleNames, EInstallBundleRequestFlags AddFlags, EInstallBundleRequestFlags RemoveFlags) = 0;
 	virtual void SetCellularPreference(int32 Value) = 0;
+	// Unblocks any pending downloads that are waiting for opting into using cellular data.
+	// Semantic wrapper until we get a better API. 
+	virtual void UnblockAnyDownloadsWaitingForCellularOptIn()
+	{
+		// Some platforms like Android clears cellular preference internally after some time,
+		// we need to constantly set it per request to unblock any pending downloads.
+		SetCellularPreference(1);
+	}
 
 	virtual void SetCacheSize(FName CacheName, uint64 CacheSize) = 0;
 
diff --git a/Engine/Source/Runtime/Online/BackgroundHTTP/Private/BackgroundHttpManagerImpl.cpp b/Engine/Source/Runtime/Online/BackgroundHTTP/Private/BackgroundHttpManagerImpl.cpp
--- a/Engine/Source/Runtime/Online/BackgroundHTTP/Private/BackgroundHttpManagerImpl.cpp
+++ b/Engine/Source/Runtime/Online/BackgroundHTTP/Private/BackgroundHttpManagerImpl.cpp
@@ -347,13 +347,6 @@ void FBackgroundHttpManagerImpl::CleanUpDataAfterCompletingRequest(const FBackgr
 	}
 }
 
-void FBackgroundHttpManagerImpl::SetCellularPreference(int32 Value)
-{
-#if PLATFORM_ANDROID
-	FAndroidMisc::SetCellularPreference(Value);
-#endif
-}
-
 bool FBackgroundHttpManagerImpl::AssociateWithAnyExistingRequest(const FBackgroundHttpRequestPtr Request)
 {
 	bool bDidAssociateWithExistingRequest = false;
diff --git a/Engine/Source/Runtime/Online/BackgroundHTTP/Public/BackgroundHttpManagerImpl.h b/Engine/Source/Runtime/Online/BackgroundHTTP/Public/BackgroundHttpManagerImpl.h
--- a/Engine/Source/Runtime/Online/BackgroundHTTP/Public/BackgroundHttpManagerImpl.h
+++ b/Engine/Source/Runtime/Online/BackgroundHTTP/Public/BackgroundHttpManagerImpl.h
@@ -44,7 +44,7 @@ public:
 	BACKGROUNDHTTP_API virtual FString GetTempFileLocationForURL(const FString& URL) override;
 	
 	BACKGROUNDHTTP_API virtual void CleanUpDataAfterCompletingRequest(const FBackgroundHttpRequestPtr Request) override;
-	BACKGROUNDHTTP_API virtual void SetCellularPreference(int32 Value) override;
+	BACKGROUNDHTTP_API virtual void SetCellularPreference(int32 Value) override {}
 	
 	//FTSTickerObjectBase implementation
 	BACKGROUNDHTTP_API virtual bool Tick(float DeltaTime) override;
