Allow airplane mode to use wifi if reenabled by user

#android

diff --git a/Engine/Plugins/Online/Android/AndroidFetchBackgroundDownload/Source/AndroidFetchBackgroundDownload/Java/download/UEDownloadWorker.java b/Engine/Plugins/Online/Android/AndroidFetchBackgroundDownload/Source/AndroidFetchBackgroundDownload/Java/download/UEDownloadWorker.java
--- a/Engine/Plugins/Online/Android/AndroidFetchBackgroundDownload/Source/AndroidFetchBackgroundDownload/Java/download/UEDownloadWorker.java
+++ b/Engine/Plugins/Online/Android/AndroidFetchBackgroundDownload/Source/AndroidFetchBackgroundDownload/Java/download/UEDownloadWorker.java
@@ -238,84 +238,71 @@ public class UEDownloadWorker extends UEWorker implements DownloadProgressListen
 	}
 
 	// Checking what the current network type is.
-	// Will pause all downloads if the network type is Cellular
+	// Will pause all downloads if the network type is Cellular without permission
 	private void NetworkTypeCheck()
 	{
-//		boolean ShouldHandleCellular = true;
-//		if (NotificationDescription != null)
-//		{
-//			ShouldHandleCellular = true; //NotificationDescription.ShouldHandleCellular;
-//		}
+		meteredType = checkMetered();
 
-		// This should be handled by the Game thread otherwise
-//		if (!bGameThreadIsActive && ShouldHandleCellular)
-		{
-			NetworkConnectivityClient.NetworkTransportType networkType = NetworkChangedManager.getInstance().networkTransportTypeCheck();
-			meteredType = checkMetered();
+		// check if airplane mode is enabled (may be cause of lost network)
+		bAirplaneMode = checkAirplaneMode();
 
-			// check if network disabled
-			bAirplaneMode = checkAirplaneMode();
-			if (bAirplaneMode)
+		if (bGameThreadIsActive)
+		{
+			// check status but don't try to affect pause/resume state as game will handle it
+			if (bLostNetwork)
 			{
-				// allow direct ethernet connection in airplane mode
-				if (networkType == NetworkConnectivityClient.NetworkTransportType.ETHERNET)
-				{
-					bAirplaneMode = false;
-				}
+				bWaitingForCellularApproval = false;
+				return;
 			}
 
-			// Pause all downloads if we lose all network connectivity, airplane mode, or data saver enabled
-			if (bLostNetwork || bAirplaneMode || meteredType == MeteredType.METERED_DATASAVER_ON)
+			// if we are on cellular we need permission
+			NetworkConnectivityClient.NetworkTransportType networkType = NetworkChangedManager.getInstance().networkTransportTypeCheck();
+			if (networkType == NetworkConnectivityClient.NetworkTransportType.CELLULAR)
 			{
-				if (!bWaitingForNetwork)
+				// check if we have permission granted
+				if (!checkCellularAllowed())
 				{
-					// wait for network to return and pause all downloads
-					bWaitingForNetwork = true;
 					bWaitingForCellularApproval = true;
-					mFetchManager.PauseAllDownloads();
+					return;
 				}
-				return;
+
 			}
 
-			// Need permission for using cellular
-			if (networkType == NetworkConnectivityClient.NetworkTransportType.CELLULAR)
+			// working network and not waiting for cellular
+			bWaitingForCellularApproval = false;
+		}
+		else
+		{
+			// if we lose network connectivity or data saver mode enabled when inactive, pause all downloads
+			if (bLostNetwork || (meteredType == MeteredType.METERED_DATASAVER_ON && !bGameThreadIsActive))
 			{
-				// check if we have permission granted
-				boolean allowCell = checkCellularAllowed();
+				// Make sure all downloads are paused (low cost if already paused)
+				mFetchManager.PauseAllDownloads();
 
-				// if not already waiting, check if we have permission
-				if (!bWaitingForCellularApproval)
-				{
-					if (!allowCell)
-					{
-						// no permission so pause and wait for it
-						mFetchManager.PauseAllDownloads();
-						bWaitingForCellularApproval = true;
-					}
-				}
-				else
-				{
-					// see if we got permission
-					if (allowCell)
-					{
-						// permission granted so resume download
-						mFetchManager.ResumeAllDownloads();
-						bWaitingForNetwork = false;
-						bWaitingForCellularApproval = false;
-					}
-				}
+				// not waiting for cell if no active network allowing download
+				bWaitingForCellularApproval = false;
+
+				return;
 			}
-			else
+
+			// if we are on cellular we need permission
+			NetworkConnectivityClient.NetworkTransportType networkType = NetworkChangedManager.getInstance().networkTransportTypeCheck();
+			if (networkType == NetworkConnectivityClient.NetworkTransportType.CELLULAR)
 			{
-				// have a network connection that is not cellular
-				if (bWaitingForNetwork || bWaitingForCellularApproval)
+				// check if we have permission granted
+				if (!checkCellularAllowed())
 				{
-					// resume any downloads
-					mFetchManager.ResumeAllDownloads();
-					bWaitingForNetwork = false;
-					bWaitingForCellularApproval = false;
+					// we need to pause and wait for approval (low cost if already paused)
+					mFetchManager.PauseAllDownloads();
+					bWaitingForCellularApproval = true;
+
+					return;
 				}
 			}
+
+			// working network and not blocked so make sure downloads running (low cost if already running)
+			mFetchManager.ResumeAllDownloads();
+			bWaitingForCellularApproval = false;
 		}
 	}
 	
@@ -413,11 +400,11 @@ public class UEDownloadWorker extends UEWorker implements DownloadProgressListen
 		}
 		
 		Notification notification = null;
-		if (bAirplaneMode)
+		if (bAirplaneMode && bLostNetwork)
 		{
 			notification = CreateAirplaneModeNotification(context, Description, pendingNotificationIntent);
 		}
-		else if (meteredType == MeteredType.METERED_DATASAVER_ON)
+		else if (meteredType == MeteredType.METERED_DATASAVER_ON && !bGameThreadIsActive)
 		{
 			notification = CreateDataSaverEnabledNotification(context, Description, pendingNotificationIntent);
 		}
@@ -823,6 +810,8 @@ public class UEDownloadWorker extends UEWorker implements DownloadProgressListen
 	//
 	public void PauseRequest(String RequestID)
 	{
+		Log.debug("@*@ PauseRequest: " + RequestID);
+
 		boolean ShouldHandleCellular = false;
 		if (NotificationDescription != null)
 		{
@@ -831,7 +820,7 @@ public class UEDownloadWorker extends UEWorker implements DownloadProgressListen
 		if (ShouldHandleCellular)
 		{
 			NetworkConnectivityClient.NetworkTransportType networkType = NetworkChangedManager.getInstance().networkTransportTypeCheck();
-			if (networkType == NetworkConnectivityClient.NetworkTransportType.CELLULAR)
+			if (networkType == NetworkConnectivityClient.NetworkTransportType.CELLULAR && !checkCellularAllowed())
 			{
 				bWaitingForCellularApproval = true;
 			}
@@ -841,6 +830,8 @@ public class UEDownloadWorker extends UEWorker implements DownloadProgressListen
 	
 	public void ResumeRequest(String RequestID)
 	{
+		Log.debug("@*@ ResumeRequest: " + RequestID);
+
 		// When the C++ resumes the request all network issues should be solved
 		bWaitingForCellularApproval = false;
 		mFetchManager.PauseDownload(RequestID, false);
@@ -874,7 +865,6 @@ public class UEDownloadWorker extends UEWorker implements DownloadProgressListen
 	private boolean bLostNetwork = false;
 	private boolean bAirplaneMode = false;
 	private MeteredType meteredType = MeteredType.UNMETERED;
-	private boolean bWaitingForNetwork = false;
 	private boolean bForceStopped = false;
 	private DownloadQueueDescription QueueDescription = null;
 	private volatile boolean bHasEnqueueHappened = false;
diff --git a/Engine/Plugins/Online/Android/AndroidFetchBackgroundDownload/Source/AndroidFetchBackgroundDownload/Java/download/fetch/FetchManager.java b/Engine/Plugins/Online/Android/AndroidFetchBackgroundDownload/Source/AndroidFetchBackgroundDownload/Java/download/fetch/FetchManager.java
--- a/Engine/Plugins/Online/Android/AndroidFetchBackgroundDownload/Source/AndroidFetchBackgroundDownload/Java/download/fetch/FetchManager.java
+++ b/Engine/Plugins/Online/Android/AndroidFetchBackgroundDownload/Source/AndroidFetchBackgroundDownload/Java/download/fetch/FetchManager.java
@@ -49,6 +49,13 @@ import java.io.File;
 //Class that handles setting up and managing Fetch2 requests
 public class FetchManager implements FetchDownloadProgressOwner, FetchEnqueueResultListener
 {
+	private enum ETaskRunningState
+	{
+		Unknown,
+		Running,
+		Paused
+	};
+
 	public FetchManager()
 	{
 		RequestedDownloads =  new HashMap<String,DownloadDescription>();
@@ -57,6 +64,7 @@ public class FetchManager implements FetchDownloadProgressOwner, FetchEnqueueRes
 	public synchronized void StopWork(String WorkID)
 	{
 		Log.debug("StopWork called");
+		runningState = ETaskRunningState.Unknown;
 
 		//Immediately remove FetchListener to prevent us from getting late callbacks after work has stopped
 		if (FetchListener != null)
@@ -83,6 +91,7 @@ public class FetchManager implements FetchDownloadProgressOwner, FetchEnqueueRes
 	
 	private synchronized void DeleteAllInProgressRequests()
 	{
+		runningState = ETaskRunningState.Unknown;
 		if (!IsFetchInstanceValid())
 		{
 			Log.error("Call to DeleteAllInProgressRequests with an unexpected invalid FetchInstance");
@@ -98,6 +107,7 @@ public class FetchManager implements FetchDownloadProgressOwner, FetchEnqueueRes
 
 	public synchronized void EnqueueRequests(Context context, DownloadQueueDescription QueueDescription)
 	{
+		runningState = ETaskRunningState.Unknown;
 		InitFetch(context);
 		
 		int NumDownloads = QueueDescription.DownloadDescriptions.size();
@@ -240,6 +250,8 @@ public class FetchManager implements FetchDownloadProgressOwner, FetchEnqueueRes
 
 	public synchronized void PauseDownload(String RequestID, boolean bPause)
 	{
+		runningState = ETaskRunningState.Unknown;
+
 		//Early out if we have invalidated our FetchInstance and thus would crash if trying to do work
 		if (!IsFetchInstanceValid())
 		{
@@ -267,11 +279,21 @@ public class FetchManager implements FetchDownloadProgressOwner, FetchEnqueueRes
 
 	public synchronized void PauseAllDownloads()
 	{
+		if (runningState == ETaskRunningState.Paused)
+		{
+			Log.debug("Fetch - PauseAllDownloads - Already paused!");
+			return;
+		}
+
 		if (!IsFetchInstanceValid())
 		{
+			Log.debug("Fetch - PauseAllDownloads - Not valid!");
 			return;
 		}
-		
+
+		int pauseCount = 0;
+		int totalCount = 0;
+
 		HashMap<String,DownloadDescription> TempRequestedDownloadsCopy =  new HashMap<String,DownloadDescription>(RequestedDownloads);
 		for(Map.Entry<String, DownloadDescription> entry : RequestedDownloads.entrySet()) 
 		{
@@ -286,16 +308,33 @@ public class FetchManager implements FetchDownloadProgressOwner, FetchEnqueueRes
 			{
 				MatchedDesc.bIsPaused = true;
 				FetchInstance.pause(MatchedDesc.CachedFetchID);
+
+				pauseCount++;
 			}
+
+			totalCount++;
 		}
+
+		runningState = ETaskRunningState.Paused;
+		Log.debug("Fetch - PauseAllDownloads - Paused " + pauseCount + " of " + totalCount);
 	}
 
 	public synchronized void ResumeAllDownloads()
 	{
+		if (runningState == ETaskRunningState.Running)
+		{
+			Log.debug("Fetch - ResumeAllDownloads - Already resumed!");
+			return;
+		}
+		
 		if (!IsFetchInstanceValid())
 		{
+			Log.debug("Fetch - ResumeAllDownloads - Not valid!");
 			return;
 		}
+
+		int resumeCount = 0;
+		int totalCount = 0;
 		
 		HashMap<String,DownloadDescription> TempRequestedDownloadsCopy =  new HashMap<String,DownloadDescription>(RequestedDownloads);
 		for(Map.Entry<String, DownloadDescription> entry : RequestedDownloads.entrySet()) 
@@ -311,8 +350,15 @@ public class FetchManager implements FetchDownloadProgressOwner, FetchEnqueueRes
 			{
 				MatchedDesc.bIsPaused = false;
 				FetchInstance.resume(MatchedDesc.CachedFetchID);
+
+				resumeCount++;
 			}
+
+			totalCount++;
 		}
+
+		runningState = ETaskRunningState.Running;
+		Log.debug("Fetch - ResumeAllDownloads - Resumed " + resumeCount + " of " + totalCount);
 	}
 
 
@@ -743,6 +789,8 @@ public class FetchManager implements FetchDownloadProgressOwner, FetchEnqueueRes
 	@Override
 	public synchronized void OnDownloadChangePauseState(@NonNull Download download, boolean bIsPaused)
 	{
+		runningState = ETaskRunningState.Unknown;
+
 		String DownloadRequestID = GetRequestID(download);
 		Log.verbose("OnDownloadChangePauseState: " + DownloadRequestID);
 
@@ -1341,6 +1389,8 @@ public class FetchManager implements FetchDownloadProgressOwner, FetchEnqueueRes
 	private volatile OkHttpClient okHttpClient = null;
 	private volatile boolean AllowHTTP2 = true;
 
+	private volatile ETaskRunningState runningState = ETaskRunningState.Unknown;
+
 	private volatile HashMap<String, DownloadDescription> RequestedDownloads = new HashMap<String, DownloadDescription>();
 	private volatile HashMap<String, DownloadDescription> CompletedDownloads = new HashMap<String, DownloadDescription>();
 	private volatile HashMap<String, DownloadDescription> FailedDownloads = new HashMap<String, DownloadDescription>();
