Fix issue with early reuse of download task filename

#android

diff --git a/Engine/Plugins/Online/Android/AndroidFetchBackgroundDownload/Source/AndroidFetchBackgroundDownload/Private/AndroidPlatformBackgroundHttpManager.cpp b/Engine/Plugins/Online/Android/AndroidFetchBackgroundDownload/Source/AndroidFetchBackgroundDownload/Private/AndroidPlatformBackgroundHttpManager.cpp
--- a/Engine/Plugins/Online/Android/AndroidFetchBackgroundDownload/Source/AndroidFetchBackgroundDownload/Private/AndroidPlatformBackgroundHttpManager.cpp
+++ b/Engine/Plugins/Online/Android/AndroidFetchBackgroundDownload/Source/AndroidFetchBackgroundDownload/Private/AndroidPlatformBackgroundHttpManager.cpp
@@ -660,31 +660,6 @@ void FAndroidPlatformBackgroundHttpManager::HandleRequestsWaitingOnJavaThreadsaf
 }
 
 const FString FAndroidPlatformBackgroundHttpManager::GetFullFileNameForDownloadDescriptionList() const
-{
-	//We should never really hit this bundle, but worth inserting a cap on how many we check to sanity deletion behavior
-	static const int MAX_NUM_DOWNLOAD_DESC_FILES = 10;
-	
-	//Find the next available filename
-	int AppendedFileNameInt = 0;
-	for (AppendedFileNameInt = 0; AppendedFileNameInt < MAX_NUM_DOWNLOAD_DESC_FILES; ++AppendedFileNameInt)
-	{
-		const FString FileToCheck = GetBaseFileNameForDownloadDescriptionListWithAppendedInt(AppendedFileNameInt);
-		if (!FPaths::FileExists(FileToCheck))
-		{
-			break;
-		}
-	}
-
-	if (!ensureAlwaysMsgf((AppendedFileNameInt < MAX_NUM_DOWNLOAD_DESC_FILES), TEXT("DownloadDescriptionList folder full of files! May lead to cases where we stomp expected .ini files for other workers!")))
-	{
-		static int StompNum = 0;
-		AppendedFileNameInt = (StompNum++ % MAX_NUM_DOWNLOAD_DESC_FILES);
-	}
-	
-	return GetBaseFileNameForDownloadDescriptionListWithAppendedInt(AppendedFileNameInt);;
-}
-
-const FString FAndroidPlatformBackgroundHttpManager::GetBaseFileNameForDownloadDescriptionListWithAppendedInt(int IntToAppend) const
 {
 	static const FString BASE_FILE_NAME = TEXT("DownloadDescriptionListJSON");
 	static const FString FILE_EXTENSION = TEXT(".ini");
@@ -692,8 +667,7 @@ const FString FAndroidPlatformBackgroundHttpManager::GetBaseFileNameForDownloadD
 	static const FString RootBGTempDir = GetFileHashHelper()->GetTemporaryRootPath();
 	static const FString BGDownloadDescriptionFolder = FPaths::Combine(RootBGTempDir, TEXT("DownloadDescriptionJSONs"));
 
-	FString FileName = FString::Printf(TEXT("%s%d%s"), *BASE_FILE_NAME, IntToAppend, *FILE_EXTENSION);
-	
+	FString FileName = FString::Printf(TEXT("%s%s%s"), *BASE_FILE_NAME, *FGuid::NewGuid().ToString(), *FILE_EXTENSION);
 	return FPaths::Combine(BGDownloadDescriptionFolder, FileName);
 }
 
diff --git a/Engine/Plugins/Online/Android/AndroidFetchBackgroundDownload/Source/AndroidFetchBackgroundDownload/Public/AndroidPlatformBackgroundHttpManager.h b/Engine/Plugins/Online/Android/AndroidFetchBackgroundDownload/Source/AndroidFetchBackgroundDownload/Public/AndroidPlatformBackgroundHttpManager.h
--- a/Engine/Plugins/Online/Android/AndroidFetchBackgroundDownload/Source/AndroidFetchBackgroundDownload/Public/AndroidPlatformBackgroundHttpManager.h
+++ b/Engine/Plugins/Online/Android/AndroidFetchBackgroundDownload/Source/AndroidFetchBackgroundDownload/Public/AndroidPlatformBackgroundHttpManager.h
@@ -68,7 +68,6 @@ protected:
 	void HandleJavaWorkerReconciliation();
 
 	const FString GetFullFileNameForDownloadDescriptionList() const;
-	const FString GetBaseFileNameForDownloadDescriptionListWithAppendedInt(int IntToAppend) const;
 
 	//Handlers for our download progressing in the underlying java implementation so that we can bubble it up to UE code.
 	void Java_OnWorkerStart(FString WorkID, jobject UnderlyingWorker);
