Add detection of SQL lock error on enqueue (can happen after deleteAll)

#android

diff --git a/Engine/Plugins/Online/Android/AndroidFetchBackgroundDownload/Source/AndroidFetchBackgroundDownload/Java/download/fetch/FetchManager.java b/Engine/Plugins/Online/Android/AndroidFetchBackgroundDownload/Source/AndroidFetchBackgroundDownload/Java/download/fetch/FetchManager.java
--- a/Engine/Plugins/Online/Android/AndroidFetchBackgroundDownload/Source/AndroidFetchBackgroundDownload/Java/download/fetch/FetchManager.java
+++ b/Engine/Plugins/Online/Android/AndroidFetchBackgroundDownload/Source/AndroidFetchBackgroundDownload/Java/download/fetch/FetchManager.java
@@ -888,6 +888,14 @@ public class FetchManager implements FetchDownloadProgressOwner, FetchEnqueueRes
 	public synchronized void OnFetchEnqueueErrorCallback(@NonNull String RequestID, @NonNull Error EnqueueError)
 	{
 		Log.error("Error Enqueing Request! " + RequestID + " Error: " + EnqueueError);
+
+		Throwable thrown = EnqueueError.getThrowable();
+		if (thrown != null) {
+			String message = thrown.getMessage();
+			if (message != null && message.contains("database is locked")) {
+				Log.error("Enqueue Error due to locked SQLITE database, likely from previous deleteAll request");
+			}
+		}	
 		
 		//If no real error supplied, lets check for fetch having already enqueued this in a separate listener
 		if (IsFetchInstanceValid() && IsErrorPossiblyIndicatingPreviouslyQueued(EnqueueError))
