Update Fetch2 3.3.0 with new version of androidx.room 2.7.1

#android

diff --git a/Engine/Plugins/Online/Android/AndroidFetchBackgroundDownload/Source/AndroidFetchBackgroundDownload/AndroidFetchBackgroundDownload_UPL.xml b/Engine/Plugins/Online/Android/AndroidFetchBackgroundDownload/Source/AndroidFetchBackgroundDownload/AndroidFetchBackgroundDownload_UPL.xml
--- a/Engine/Plugins/Online/Android/AndroidFetchBackgroundDownload/Source/AndroidFetchBackgroundDownload/AndroidFetchBackgroundDownload_UPL.xml
+++ b/Engine/Plugins/Online/Android/AndroidFetchBackgroundDownload/Source/AndroidFetchBackgroundDownload/AndroidFetchBackgroundDownload_UPL.xml
@@ -15,9 +15,25 @@
 	}
 	dependencies {
 		//Fetch2 library used to manage downloads at java level
-		implementation "com.github.tonyofrancis:fetch:3.3.0" 
-		api "com.squareup.okhttp3:okhttp:3.12.12"
-		api "com.squareup.okhttp3:okhttp-urlconnection:3.12.12"
+	    api "org.jetbrains.kotlin:kotlin-stdlib:2.0.20"
+		api "androidx.room:room-runtime:2.7.1"
+		api "androidx.core:core-ktx:1.13.1"
+		api "com.squareup.okhttp3:okhttp:4.12.0"
+		api "com.squareup.okhttp3:okhttp-urlconnection:4.12.0"
+		</insert>		
+		<insertValue value="implementation files ('$S(AbsPluginDir)/ThirdParty/fetch2-release.aar')"/>
+		<insertNewline/>
+		<insertValue value="implementation files ('$S(AbsPluginDir)/ThirdParty/fetch2core-release.aar')"/>
+		<insertNewline/>
+		<insertValue value="implementation files ('$S(AbsPluginDir)/ThirdParty/fetch2fileserver-release.aar')"/>
+		<insertNewline/>
+		<insertValue value="implementation files ('$S(AbsPluginDir)/ThirdParty/fetch2okhttp-release.aar')"/>
+		<insertNewline/>
+		<insertValue value="implementation files ('$S(AbsPluginDir)/ThirdParty/fetch2rx-release.aar')"/>
+		<insertNewline/>
+		<insertValue value="implementation files ('$S(AbsPluginDir)/ThirdParty/fetchmigrator-release.aar')"/>
+		<insertNewline/>
+		<insert>
 		
 		//Used for parsing from JSON files in DownloadDescription.java
 		//excluding hamcrest to avoid collision with junit that also includes hamcrest
