Sending a background mobile funnel event for all content downloads completed

#rnx

diff --git a/Engine/Plugins/Online/Android/AndroidFetchBackgroundDownload/Source/AndroidFetchBackgroundDownload/Java/download/UEDownloadWorker.java b/Engine/Plugins/Online/Android/AndroidFetchBackgroundDownload/Source/AndroidFetchBackgroundDownload/Java/download/UEDownloadWorker.java
--- a/Engine/Plugins/Online/Android/AndroidFetchBackgroundDownload/Source/AndroidFetchBackgroundDownload/Java/download/UEDownloadWorker.java
+++ b/Engine/Plugins/Online/Android/AndroidFetchBackgroundDownload/Source/AndroidFetchBackgroundDownload/Java/download/UEDownloadWorker.java
@@ -137,7 +137,7 @@ public class UEDownloadWorker extends UEWorker implements DownloadProgressListen
 		
 		if (null == mFetchManager)
 		{
-			mFetchManager = new FetchManager();	
+			mFetchManager = FetchManager.GetSharedManager();	
 		}
 		
 		//Make sure we have a CancelIntent we can use to cancel this job (passed into notifications, etc)
diff --git a/Engine/Plugins/Online/Android/AndroidFetchBackgroundDownload/Source/AndroidFetchBackgroundDownload/Java/download/fetch/FetchManager.java b/Engine/Plugins/Online/Android/AndroidFetchBackgroundDownload/Source/AndroidFetchBackgroundDownload/Java/download/fetch/FetchManager.java
--- a/Engine/Plugins/Online/Android/AndroidFetchBackgroundDownload/Source/AndroidFetchBackgroundDownload/Java/download/fetch/FetchManager.java
+++ b/Engine/Plugins/Online/Android/AndroidFetchBackgroundDownload/Source/AndroidFetchBackgroundDownload/Java/download/fetch/FetchManager.java
@@ -55,12 +55,17 @@ public class FetchManager implements FetchDownloadProgressOwner, FetchEnqueueRes
 		Running,
 		Paused
 	};
-
-	public FetchManager()
+	
+	private FetchManager()
 	{
 		RequestedDownloads =  new HashMap<String,DownloadDescription>();
 	}
 	
+	public static synchronized FetchManager GetSharedManager()
+	{
+		return SharedManager;
+	}
+	
 	public synchronized void StopWork(String WorkID)
 	{
 		Log.debug("StopWork called");
@@ -1317,6 +1322,11 @@ public class FetchManager implements FetchDownloadProgressOwner, FetchEnqueueRes
 	{
 		boolean bDidAllSucceed = ((bForceFailure == false) && (FailedDownloads.isEmpty()));
 		ProgressListener.OnAllDownloadsComplete(bDidAllSucceed);
+		
+		if (AllDownloadsCompletedListener != null)
+		{
+			AllDownloadsCompletedListener.OnAllDownloadsCompleted(bDidAllSucceed);
+		}
 	}
 	
 	private String GetRequestID(Download download)
@@ -1351,6 +1361,11 @@ public class FetchManager implements FetchDownloadProgressOwner, FetchEnqueueRes
 		}
 	}
 	
+	public synchronized void SetAllDownloadsCompletedListener(@NonNull FetchAllDownloadsCompletedListener Listener)
+	{
+		AllDownloadsCompletedListener = Listener;
+	}
+	
 	//Wrapper class for our FetchInstance logging that passes the values through UE's Logger Class
 	//FetchLogger only has debug and error logging currently implemented so that is all we pass through.
 	private class FUEFetchLogger extends FetchLogger
@@ -1390,9 +1405,13 @@ public class FetchManager implements FetchDownloadProgressOwner, FetchEnqueueRes
 			UELogger.error(message, throwable);
 		}
 	}
+	
+	private static volatile FetchManager SharedManager = new FetchManager();
+	
 	private volatile Fetch FetchInstance = null;
 
 	private volatile FetchRequestProgressListener FetchListener = null;
+	private volatile FetchAllDownloadsCompletedListener AllDownloadsCompletedListener = null;
 
 	private volatile OkHttpClient okHttpClient = null;
 	private volatile boolean AllowHTTP2 = true;
@@ -1407,4 +1426,9 @@ public class FetchManager implements FetchDownloadProgressOwner, FetchEnqueueRes
 	
 	public Logger Log = new Logger("UE", "FetchManager");
 	public FUEFetchLogger FetchLog = new FUEFetchLogger();
+	
+	public interface FetchAllDownloadsCompletedListener
+	{
+		public void OnAllDownloadsCompleted(boolean bDidAllSucceed);
+	}
 }
\ No newline at end of file
diff --git a/Engine/Source/Runtime/InstallBundleManager/Private/InstallBundleManagerInterface.cpp b/Engine/Source/Runtime/InstallBundleManager/Private/InstallBundleManagerInterface.cpp
--- a/Engine/Source/Runtime/InstallBundleManager/Private/InstallBundleManagerInterface.cpp
+++ b/Engine/Source/Runtime/InstallBundleManager/Private/InstallBundleManagerInterface.cpp
@@ -18,6 +18,8 @@ FInstallBundleManagerOnPatchCheckComplete IInstallBundleManager::PatchCheckCompl
 
 FInstallBundleManagerReportingDelegate IInstallBundleManager::ReportingDelegate;
 
+FInstallBundleManagerAllDownloadsCompleted IInstallBundleManager::AllDownloadsCompletedDelegate;
+
 TSharedPtr<IInstallBundleManager> IInstallBundleManager::GetPlatformInstallBundleManager()
 {
 	static IInstallBundleManagerModule* Module = nullptr;
diff --git a/Engine/Source/Runtime/InstallBundleManager/Public/InstallBundleManagerInterface.h b/Engine/Source/Runtime/InstallBundleManager/Public/InstallBundleManagerInterface.h
--- a/Engine/Source/Runtime/InstallBundleManager/Public/InstallBundleManagerInterface.h
+++ b/Engine/Source/Runtime/InstallBundleManager/Public/InstallBundleManagerInterface.h
@@ -117,6 +117,7 @@ DECLARE_DELEGATE_RetVal(bool, FInstallBundleManagerEnvironmentWantsPatchCheck);
 DECLARE_DELEGATE_OneParam(FInstallBundleGetInstallStateDelegate, FInstallBundleCombinedInstallState);
 
 DECLARE_DELEGATE(FInstallBundleManagerFlushCacheCompleteDelegate);
+DECLARE_DELEGATE(FInstallBundleManagerAllDownloadsCompleted);
 
 class IInstallBundleManager : public TSharedFromThis<IInstallBundleManager>
 {
@@ -129,6 +130,7 @@ public:
 	static INSTALLBUNDLEMANAGER_API FInstallBundleReleasedMultiDelegate ReleasedDelegate; // Called when content release request is complete
 	static INSTALLBUNDLEMANAGER_API FInstallBundleManagerOnPatchCheckComplete PatchCheckCompleteDelegate;
 	static INSTALLBUNDLEMANAGER_API FInstallBundleManagerReportingDelegate ReportingDelegate;
+	static INSTALLBUNDLEMANAGER_API FInstallBundleManagerAllDownloadsCompleted AllDownloadsCompletedDelegate; // Called when all content downloads completed
 
 	static INSTALLBUNDLEMANAGER_API TSharedPtr<IInstallBundleManager> GetPlatformInstallBundleManager();
 
