Add FetchStartTimeUTC and FetchEndTimeUTC to ChunkDownloadComplete

#rnx

diff --git a/Engine/Plugins/Online/Android/AndroidFetchBackgroundDownload/Source/AndroidFetchBackgroundDownload/Java/download/DownloadProgressListener.java b/Engine/Plugins/Online/Android/AndroidFetchBackgroundDownload/Source/AndroidFetchBackgroundDownload/Java/download/DownloadProgressListener.java
--- a/Engine/Plugins/Online/Android/AndroidFetchBackgroundDownload/Source/AndroidFetchBackgroundDownload/Java/download/DownloadProgressListener.java
+++ b/Engine/Plugins/Online/Android/AndroidFetchBackgroundDownload/Source/AndroidFetchBackgroundDownload/Java/download/DownloadProgressListener.java
@@ -9,7 +9,7 @@ public interface DownloadProgressListener
 	public void OnDownloadProgress(String RequestID, long BytesWrittenSinceLastCall, long TotalBytesWritten);
 	public void OnDownloadGroupProgress(int GroupID, int Progress, boolean Indeterminate);
 	public void OnDownloadComplete(String RequestID, String CompleteLocation, EDownloadCompleteReason CompleteReason);
-	public void OnDownloadMetrics(String RequestID, long TotalBytesDownloaded, long DownloadDuration);
+	public void OnDownloadMetrics(String RequestID, long TotalBytesDownloaded, long DownloadDuration, long DownloadStartTimeUTC, long DownloadEndTimeUTC);
 	public void OnAllDownloadsComplete(boolean bDidAllRequestsSucceed);
 	public void OnDownloadEnqueued(String RequestID, boolean bEnqueueSuccess);
 }
\ No newline at end of file
diff --git a/Engine/Plugins/Online/Android/AndroidFetchBackgroundDownload/Source/AndroidFetchBackgroundDownload/Java/download/UEDownloadWorker.java b/Engine/Plugins/Online/Android/AndroidFetchBackgroundDownload/Source/AndroidFetchBackgroundDownload/Java/download/UEDownloadWorker.java
--- a/Engine/Plugins/Online/Android/AndroidFetchBackgroundDownload/Source/AndroidFetchBackgroundDownload/Java/download/UEDownloadWorker.java
+++ b/Engine/Plugins/Online/Android/AndroidFetchBackgroundDownload/Source/AndroidFetchBackgroundDownload/Java/download/UEDownloadWorker.java
@@ -810,11 +810,11 @@ public class UEDownloadWorker extends UEWorker implements DownloadProgressListen
 	}
 
 	@Override
-	public void OnDownloadMetrics(String RequestID, long TotalBytesDownloaded, long DownloadDuration)
+	public void OnDownloadMetrics(String RequestID, long TotalBytesDownloaded, long DownloadDuration, long DownloadStartTimeUTC, long DownloadEndTimeUTC)
 	{
 		try
 		{
-			nativeAndroidBackgroundDownloadOnMetrics(RequestID, TotalBytesDownloaded, DownloadDuration);
+			nativeAndroidBackgroundDownloadOnMetrics(RequestID, TotalBytesDownloaded, DownloadDuration, DownloadStartTimeUTC, DownloadEndTimeUTC);
 		}
 		catch (UnsatisfiedLinkError e)
 		{
@@ -957,7 +957,7 @@ public class UEDownloadWorker extends UEWorker implements DownloadProgressListen
 	native void nativeAndroidBackgroundDownloadOnWorkerStop(String WorkID);
 	native void nativeAndroidBackgroundDownloadOnProgress(String TaskID, long BytesWrittenSinceLastCall, long TotalBytesWritten);
 	native void nativeAndroidBackgroundDownloadOnComplete(String TaskID, String CompleteLocation, boolean bWasSuccess);
-	native void nativeAndroidBackgroundDownloadOnMetrics(String TaskID, long TotalBytesDownloaded, long DownloadDuration);
+	native void nativeAndroidBackgroundDownloadOnMetrics(String TaskID, long TotalBytesDownloaded, long DownloadDuration, long DownloadStartTimeUTC, long DownloadEndTimeUTC);
 	native void nativeAndroidBackgroundDownloadOnAllComplete(boolean bDidAllRequestsSucceed);
 	native void nativeAndroidBackgroundDownloadOnTick();
 	
diff --git a/Engine/Plugins/Online/Android/AndroidFetchBackgroundDownload/Source/AndroidFetchBackgroundDownload/Java/download/fetch/FetchManager.java b/Engine/Plugins/Online/Android/AndroidFetchBackgroundDownload/Source/AndroidFetchBackgroundDownload/Java/download/fetch/FetchManager.java
--- a/Engine/Plugins/Online/Android/AndroidFetchBackgroundDownload/Source/AndroidFetchBackgroundDownload/Java/download/fetch/FetchManager.java
+++ b/Engine/Plugins/Online/Android/AndroidFetchBackgroundDownload/Source/AndroidFetchBackgroundDownload/Java/download/fetch/FetchManager.java
@@ -1271,14 +1271,16 @@ public class FetchManager implements FetchDownloadProgressOwner, FetchEnqueueRes
 		}
 
 		long DownloadDuration = -1;
+		long DownloadFinishTime = System.currentTimeMillis();
+
 		if (DownloadDesc.DownloadStartTime > 0)
 		{
-			long ElapsedDownloadTime = System.currentTimeMillis() - DownloadDesc.DownloadStartTime;
+			long ElapsedDownloadTime = DownloadFinishTime - DownloadDesc.DownloadStartTime;
 
 			DownloadDuration = ElapsedDownloadTime - DownloadDesc.TotalPausedTime;
 		}
 
-		DownloadDesc.ProgressListener.OnDownloadMetrics(DownloadDesc.RequestID, DownloadDesc.TotalDownloadedBytes, DownloadDuration);
+		DownloadDesc.ProgressListener.OnDownloadMetrics(DownloadDesc.RequestID, DownloadDesc.TotalDownloadedBytes, DownloadDuration, DownloadDesc.DownloadStartTime, DownloadFinishTime);
 
 		DownloadDesc.ProgressListener.OnDownloadComplete(DownloadDesc.RequestID, DownloadDesc.DestinationLocation, CompleteReasonToSend);
 		
diff --git a/Engine/Plugins/Online/Android/AndroidFetchBackgroundDownload/Source/AndroidFetchBackgroundDownload/Private/AndroidPlatformBackgroundHttpManager.cpp b/Engine/Plugins/Online/Android/AndroidFetchBackgroundDownload/Source/AndroidFetchBackgroundDownload/Private/AndroidPlatformBackgroundHttpManager.cpp
--- a/Engine/Plugins/Online/Android/AndroidFetchBackgroundDownload/Source/AndroidFetchBackgroundDownload/Private/AndroidPlatformBackgroundHttpManager.cpp
+++ b/Engine/Plugins/Online/Android/AndroidFetchBackgroundDownload/Source/AndroidFetchBackgroundDownload/Private/AndroidPlatformBackgroundHttpManager.cpp
@@ -58,7 +58,7 @@ namespace FAndroidBackgroundDownloadDelegates
 	DECLARE_DELEGATE_TwoParams(FAndroidBackgroundDownload_OnWorkerStop, FString /*WorkID*/, jobject /*UnderlyingWorker*/);
 	DECLARE_DELEGATE_FourParams(FAndroidBackgroundDownload_OnProgress, jobject /*UnderlyingWorker*/, FString /*RequestID*/, int64_t /*BytesWrittenSinceLastCall*/, int64_t /*TotalBytesWritten*/);
 	DECLARE_DELEGATE_FourParams(FAndroidBackgroundDownload_OnComplete, jobject /*UnderlyingWorker*/, FString /*RequestID*/, FString /*CompleteLocation*/, bool /*bWasSuccess*/);
-	DECLARE_DELEGATE_FourParams(FAndroidBackgroundDownload_OnMetrics, jobject /*UnderlyingWorker*/, FString /*RequestID*/, int64_t /*TotalBytesDownloaded*/, int64_t /*DownloadDuration*/);
+	DECLARE_DELEGATE_SixParams(FAndroidBackgroundDownload_OnMetrics, jobject /*UnderlyingWorker*/, FString /*RequestID*/, int64_t /*TotalBytesDownloaded*/, int64_t /*DownloadDuration*/, int64_t /*DownloadStartTimeUTC*/, int64_t /*DownloadEndTimeUTC*/);
 	DECLARE_DELEGATE_TwoParams(FAndroidBackgroundDownload_OnAllComplete, jobject /*UnderlyingWorker*/, bool /*bDidAllRequestsSucceed*/);
 	DECLARE_DELEGATE_TwoParams(FAndroidBackgroundDownload_OnTickWorkerThread, JNIEnv*, jobject /*UnderlyingWorker*/);
 
@@ -696,7 +696,7 @@ void FAndroidPlatformBackgroundHttpManager::Java_OnDownloadComplete(jobject Unde
 	}
 }
 
-void FAndroidPlatformBackgroundHttpManager::Java_OnDownloadMetrics(jobject UnderlyingWorker, FString RequestID, int64_t TotalBytesDownloaded, int64_t DownloadDuration)
+void FAndroidPlatformBackgroundHttpManager::Java_OnDownloadMetrics(jobject UnderlyingWorker, FString RequestID, int64_t TotalBytesDownloaded, int64_t DownloadDuration, int64_t DownloadStartTimeUTC, int64_t DownloadEndTimeUTC)
 {
 	UE_LOG(LogBackgroundHttpManager, Log, TEXT("DownloadMetrics... RequestID:%s TotalBytesDownloaded:%ld DownloadDuration:%ld"), *RequestID, TotalBytesDownloaded, DownloadDuration);
 
@@ -704,8 +704,10 @@ void FAndroidPlatformBackgroundHttpManager::Java_OnDownloadMetrics(jobject Under
 	if (CompletedRequest.IsValid())
 	{
 		float DownloadDurationInSeconds = (float)((double)DownloadDuration / 1000.0f);
+		FDateTime FetchStartTimeUTC = FDateTime::FromUnixTimestampDecimal((double)DownloadStartTimeUTC / 1000.0f);
+		FDateTime FetchEndTimeUTC = FDateTime::FromUnixTimestampDecimal((double)DownloadEndTimeUTC / 1000.0f);
 
-		CompletedRequest->NotifyRequestMetricsAvailable((int32)TotalBytesDownloaded, DownloadDurationInSeconds);
+		CompletedRequest->NotifyRequestMetricsAvailable((int32)TotalBytesDownloaded, DownloadDurationInSeconds, FetchStartTimeUTC, FetchEndTimeUTC);
 	}
 }
 
@@ -1256,13 +1258,15 @@ JNI_METHOD void Java_com_epicgames_unreal_download_UEDownloadWorker_nativeAndroi
 	FAndroidBackgroundDownloadDelegates::AndroidBackgroundDownload_OnComplete.ExecuteIfBound(thiz, RequestID, ConvertedCompleteLocation, ConvertedbWasSuccess);
 }
 
-JNI_METHOD void Java_com_epicgames_unreal_download_UEDownloadWorker_nativeAndroidBackgroundDownloadOnMetrics(JNIEnv* jenv, jobject thiz, jstring TaskID, jlong TotalBytesDownloaded, jlong DownloadDuration)
+JNI_METHOD void Java_com_epicgames_unreal_download_UEDownloadWorker_nativeAndroidBackgroundDownloadOnMetrics(JNIEnv* jenv, jobject thiz, jstring TaskID, jlong TotalBytesDownloaded, jlong DownloadDuration, jlong DownloadStartTimeUTC, jlong DownloadEndTimeUTC)
 {
 	FString RequestID = FJavaHelper::FStringFromParam(jenv, TaskID);
 	int64_t ConvertedTotalBytesDownloaded = static_cast<uint64_t>(TotalBytesDownloaded);
 	int64_t ConvertedDownloadDuration = static_cast<uint64_t>(DownloadDuration);
+	int64_t ConvertedDownloadStartTimeUTC = static_cast<uint64_t>(DownloadStartTimeUTC);
+	int64_t ConvertedDownloadEndTimeUTC = static_cast<uint64_t>(DownloadEndTimeUTC);
 
-	FAndroidBackgroundDownloadDelegates::AndroidBackgroundDownload_OnMetrics.ExecuteIfBound(thiz, RequestID, ConvertedTotalBytesDownloaded, ConvertedDownloadDuration);
+	FAndroidBackgroundDownloadDelegates::AndroidBackgroundDownload_OnMetrics.ExecuteIfBound(thiz, RequestID, ConvertedTotalBytesDownloaded, ConvertedDownloadDuration, ConvertedDownloadStartTimeUTC, ConvertedDownloadEndTimeUTC);
 }
 
 JNI_METHOD void Java_com_epicgames_unreal_download_UEDownloadWorker_nativeAndroidBackgroundDownloadOnAllComplete(JNIEnv* jenv, jobject thiz, jboolean bDidAllRequestsSucceed)
diff --git a/Engine/Plugins/Online/Android/AndroidFetchBackgroundDownload/Source/AndroidFetchBackgroundDownload/Public/AndroidPlatformBackgroundHttpManager.h b/Engine/Plugins/Online/Android/AndroidFetchBackgroundDownload/Source/AndroidFetchBackgroundDownload/Public/AndroidPlatformBackgroundHttpManager.h
--- a/Engine/Plugins/Online/Android/AndroidFetchBackgroundDownload/Source/AndroidFetchBackgroundDownload/Public/AndroidPlatformBackgroundHttpManager.h
+++ b/Engine/Plugins/Online/Android/AndroidFetchBackgroundDownload/Source/AndroidFetchBackgroundDownload/Public/AndroidPlatformBackgroundHttpManager.h
@@ -74,7 +74,7 @@ protected:
 	void Java_OnWorkerStop(FString WorkID, jobject UnderlyingWorker);
 	void Java_OnDownloadProgress(jobject UnderlyingWorker, FString RequestID, int64_t BytesWrittenSinceLastCall, int64_t TotalBytesWritten);
 	void Java_OnDownloadComplete(jobject UnderlyingWorker, FString RequestID, FString CompleteLocation, bool bWasSuccess);
-	void Java_OnDownloadMetrics(jobject UnderlyingWorker, FString RequestID, int64_t TotalBytesDownloaded, int64_t DownloadDuration);
+	void Java_OnDownloadMetrics(jobject UnderlyingWorker, FString RequestID, int64_t TotalBytesDownloaded, int64_t DownloadDuration, int64_t DownloadStartTimeUTC, int64_t DownloadEndTimeUTC);
 	void Java_OnAllDownloadsComplete(jobject UnderlyingWorker, bool bDidAllRequestsSucceed);
 	void Java_OnTick(JNIEnv* Env, jobject UnderlyingWorker);
 
diff --git a/Engine/Source/Runtime/ApplicationCore/Private/IOS/IOSBackgroundURLSessionHandler.cpp b/Engine/Source/Runtime/ApplicationCore/Private/IOS/IOSBackgroundURLSessionHandler.cpp
--- a/Engine/Source/Runtime/ApplicationCore/Private/IOS/IOSBackgroundURLSessionHandler.cpp
+++ b/Engine/Source/Runtime/ApplicationCore/Private/IOS/IOSBackgroundURLSessionHandler.cpp
@@ -1030,7 +1030,8 @@ static constexpr NSInteger HTTPStatusCodeErrorServer = 500;
 			// Calculate and log download speed
             if (responseDuration > 0 && metric.countOfResponseBodyBytesReceived > 0) {
 				
-				FBackgroundURLSessionHandler::OnDownloadMetrics.Broadcast((uint64)task.taskIdentifier, (int32)metric.countOfResponseBodyBytesReceived, (float)responseDuration);
+				FBackgroundURLSessionHandler::OnDownloadMetrics.Broadcast((uint64)task.taskIdentifier, (int32)metric.countOfResponseBodyBytesReceived, (float)responseDuration, 
+					(int64_t)round([metric.responseStartDate timeIntervalSince1970] * 1000.0), (int64_t)round([metric.responseEndDate timeIntervalSince1970] * 1000.0));
 				
 #if !UE_BUILD_SHIPPING
 				// bytes per second
diff --git a/Engine/Source/Runtime/ApplicationCore/Public/IOS/IOSBackgroundURLSessionHandler.h b/Engine/Source/Runtime/ApplicationCore/Public/IOS/IOSBackgroundURLSessionHandler.h
--- a/Engine/Source/Runtime/ApplicationCore/Public/IOS/IOSBackgroundURLSessionHandler.h
+++ b/Engine/Source/Runtime/ApplicationCore/Public/IOS/IOSBackgroundURLSessionHandler.h
@@ -61,7 +61,7 @@ public:
 	static FOnDownloadCompleted OnDownloadCompleted;
 	
 	// Will be invoked from didFinishCollectingMetrics
-	DECLARE_MULTICAST_DELEGATE_ThreeParams(FOnDownloadMetrics, const uint64 /*DownloadId*/, const int32 /*TotalBytesDownloaded*/, const float /*DownloadDuration*/)
+	DECLARE_MULTICAST_DELEGATE_FiveParams(FOnDownloadMetrics, const uint64 /*DownloadId*/, const int32 /*TotalBytesDownloaded*/, const float /*DownloadDuration*/, const int64_t /* DownloadStartTime*/, const int64_t /*DownloadEndTime*/);
 	static FOnDownloadMetrics OnDownloadMetrics;
 
 	// Will be invoked from handleEventsForBackgroundURLSession application delegate. Needs to be registered very early, e.g. from static constructor.
diff --git a/Engine/Source/Runtime/InstallBundleManager/Public/InstallBundleManagerInterface.h b/Engine/Source/Runtime/InstallBundleManager/Public/InstallBundleManagerInterface.h
--- a/Engine/Source/Runtime/InstallBundleManager/Public/InstallBundleManagerInterface.h
+++ b/Engine/Source/Runtime/InstallBundleManager/Public/InstallBundleManagerInterface.h
@@ -96,6 +96,8 @@ struct FInstallBundleChunkDownloadInfo
 	FString ChunkName;
 	int32 ChunkFileSize;
 	float ChunkDownloadDuration;
+	FDateTime FetchStartTimeUTC;
+	FDateTime FetchEndTimeUTC;
 };
 
 using FInstallBundleSourceOrCache = TUnion<FInstallBundleSourceType, FName>;
diff --git a/Engine/Source/Runtime/Online/BackgroundHTTP/Private/BackgroundHttpManagerImpl.cpp b/Engine/Source/Runtime/Online/BackgroundHTTP/Private/BackgroundHttpManagerImpl.cpp
--- a/Engine/Source/Runtime/Online/BackgroundHTTP/Private/BackgroundHttpManagerImpl.cpp
+++ b/Engine/Source/Runtime/Online/BackgroundHTTP/Private/BackgroundHttpManagerImpl.cpp
@@ -364,8 +364,8 @@ bool FBackgroundHttpManagerImpl::AssociateWithAnyExistingRequest(const FBackgrou
 			//First send progress update for the file size so anything monitoring this download knows we are about to update this progress
 			Request->OnProgressUpdated().ExecuteIfBound(Request, ExistingFileSize, ExistingFileSize);
 
-			//Now complete with this completed response data
-			Request->SetMetrics(ExistingFileSize, 0.0f);
+			//Now complete with this response data.  Since we have no time data fill everything in with zeros
+			Request->SetMetrics(ExistingFileSize, 0.0f, FDateTime::FromUnixTimestamp(0), FDateTime::FromUnixTimestamp(0));
 			Request->CompleteWithExistingResponseData(NewResponseWithExistingFile);
 		}
 	}
diff --git a/Engine/Source/Runtime/Online/BackgroundHTTP/Private/BackgroundHttpRequestImpl.cpp b/Engine/Source/Runtime/Online/BackgroundHTTP/Private/BackgroundHttpRequestImpl.cpp
--- a/Engine/Source/Runtime/Online/BackgroundHTTP/Private/BackgroundHttpRequestImpl.cpp
+++ b/Engine/Source/Runtime/Online/BackgroundHTTP/Private/BackgroundHttpRequestImpl.cpp
@@ -80,7 +80,7 @@ void FBackgroundHttpRequestImpl::OnBackgroundDownloadComplete()
 		{
 			UE_LOG(LogBackgroundHttpRequest, Warning, TEXT("OptionalMetricsInfo OnRequestMetrics() is NOT bound - RequestID:%s"), *GetRequestID());
 		}
-		OnRequestMetrics().ExecuteIfBound(SharedThis(this), OptionalMetricsInfo->TotalBytesDownloaded, OptionalMetricsInfo->DownloadDuration);
+		OnRequestMetrics().ExecuteIfBound(SharedThis(this), OptionalMetricsInfo->TotalBytesDownloaded, OptionalMetricsInfo->DownloadDuration, OptionalMetricsInfo->FetchStartTimeUTC, OptionalMetricsInfo->FetchEndTimeUTC);
 		OptionalMetricsInfo.Reset();
 	}
 	else
@@ -103,17 +103,19 @@ void FBackgroundHttpRequestImpl::NotifyNotificationObjectOfComplete(bool bWasSuc
 	}
 }
 
-void FBackgroundHttpRequestImpl::SetMetrics(const int32 TotalBytesDownloaded, const float DownloadDuration)
+void FBackgroundHttpRequestImpl::SetMetrics(const int32 TotalBytesDownloaded, const float DownloadDuration, const FDateTime& FetchStartTimeUTC, const FDateTime& FetchEndTimeUTC)
 {
-	NotifyRequestMetricsAvailable(TotalBytesDownloaded, DownloadDuration);
+	NotifyRequestMetricsAvailable(TotalBytesDownloaded, DownloadDuration, FetchStartTimeUTC, FetchEndTimeUTC);
 }
 
-void FBackgroundHttpRequestImpl::NotifyRequestMetricsAvailable(const int32 TotalBytesDownloaded, const float DownloadDuration)
+void FBackgroundHttpRequestImpl::NotifyRequestMetricsAvailable(const int32 TotalBytesDownloaded, const float DownloadDuration, const FDateTime& FetchStartTimeUTC, const FDateTime& FetchEndTimeUTC)
 {
 	ensureAlwaysMsgf(!OptionalMetricsInfo.IsSet(), TEXT("Overwriting metrics data!!!"));
 	OptionalMetricsInfo = {
 		.TotalBytesDownloaded = TotalBytesDownloaded,
-		.DownloadDuration = DownloadDuration
+		.DownloadDuration = DownloadDuration,
+		.FetchStartTimeUTC = FetchStartTimeUTC,
+		.FetchEndTimeUTC = FetchEndTimeUTC
 	};
 }
 
diff --git a/Engine/Source/Runtime/Online/BackgroundHTTP/Private/IOS/ApplePlatformBackgroundHttpManager.cpp b/Engine/Source/Runtime/Online/BackgroundHTTP/Private/IOS/ApplePlatformBackgroundHttpManager.cpp
--- a/Engine/Source/Runtime/Online/BackgroundHTTP/Private/IOS/ApplePlatformBackgroundHttpManager.cpp
+++ b/Engine/Source/Runtime/Online/BackgroundHTTP/Private/IOS/ApplePlatformBackgroundHttpManager.cpp
@@ -138,7 +138,7 @@ void FApplePlatformBackgroundHttpManager::OnDownloadCompleted(const uint64 Downl
 	}
 }
 
-void FApplePlatformBackgroundHttpManager::OnDownloadMetrics(const uint64 DownloadId, const int32 TotalBytesDownloaded, const float DownloadDuration)
+void FApplePlatformBackgroundHttpManager::OnDownloadMetrics(const uint64 DownloadId, const int32 TotalBytesDownloaded, const float DownloadDuration, const int64_t DownloadStartTimeUTC, const int64_t DownloadEndTimeUTC)
 {
 	FRWScopeLock ScopeLock(ActiveRequestLock, SLT_ReadOnly);
 	for (FBackgroundHttpRequestPtr& GenericRequest : ActiveRequests)
@@ -148,7 +148,9 @@ void FApplePlatformBackgroundHttpManager::OnDownloadMetrics(const uint64 Downloa
 		{
 			if (Request->GetInternalDownloadId() == DownloadId)
 			{
-				Request->NotifyRequestMetricsAvailable(TotalBytesDownloaded, DownloadDuration);
+				FDateTime FetchStartTimeUTC = FDateTime::FromUnixTimestampDecimal((double)DownloadStartTimeUTC / 1000.0f);
+				FDateTime FetchEndTimeUTC = FDateTime::FromUnixTimestampDecimal((double)DownloadEndTimeUTC / 1000.0f);
+				Request->NotifyRequestMetricsAvailable(TotalBytesDownloaded, DownloadDuration, FetchStartTimeUTC, FetchEndTimeUTC);
 				return;
 			}
 		}
diff --git a/Engine/Source/Runtime/Online/BackgroundHTTP/Public/BackgroundHttpRequestImpl.h b/Engine/Source/Runtime/Online/BackgroundHTTP/Public/BackgroundHttpRequestImpl.h
--- a/Engine/Source/Runtime/Online/BackgroundHTTP/Public/BackgroundHttpRequestImpl.h
+++ b/Engine/Source/Runtime/Online/BackgroundHTTP/Public/BackgroundHttpRequestImpl.h
@@ -47,8 +47,8 @@ public:
 	BACKGROUNDHTTP_API virtual void SetRequestPriority(EBackgroundHTTPPriority NewPriority) override;
 	
 	BACKGROUNDHTTP_API virtual void NotifyNotificationObjectOfComplete(const bool bWasSuccess);
-	BACKGROUNDHTTP_API virtual void NotifyRequestMetricsAvailable(const int32 TotalBytesDownloaded, const float DownloadDuration);
-	BACKGROUNDHTTP_API virtual void SetMetrics(const int32 TotalBytesDownloaded, const float DownloadDuration) override;
+	BACKGROUNDHTTP_API virtual void NotifyRequestMetricsAvailable(const int32 TotalBytesDownloaded, const float DownloadDuration, const FDateTime& FetchStartTimeUTC, const FDateTime& FetchEndTimeUTC);
+	BACKGROUNDHTTP_API virtual void SetMetrics(const int32 TotalBytesDownloaded, const float DownloadDuration, const FDateTime& FetchStartTimeUTC, const FDateTime& FetchEndTimeUTC) override;
 	
 protected:
 	UE::FMutex DownloadCompleteMutex;
@@ -64,6 +64,8 @@ protected:
 	{
 		int32 TotalBytesDownloaded;
 		float DownloadDuration;
+		FDateTime FetchStartTimeUTC;
+		FDateTime FetchEndTimeUTC;
 	};
 	
 	TOptional<FDownloadMetricsInfo> OptionalMetricsInfo;
diff --git a/Engine/Source/Runtime/Online/BackgroundHTTP/Public/IOS/ApplePlatformBackgroundHttpManager.h b/Engine/Source/Runtime/Online/BackgroundHTTP/Public/IOS/ApplePlatformBackgroundHttpManager.h
--- a/Engine/Source/Runtime/Online/BackgroundHTTP/Public/IOS/ApplePlatformBackgroundHttpManager.h
+++ b/Engine/Source/Runtime/Online/BackgroundHTTP/Public/IOS/ApplePlatformBackgroundHttpManager.h
@@ -34,5 +34,5 @@ private:
 	FDelegateHandle OnDownloadCompletedHandle;
 	FDelegateHandle OnDownloadMetricsHandle;
 	void OnDownloadCompleted(const uint64 DownloadId, const bool bSuccess);
-	void OnDownloadMetrics(const uint64 DownloadId, const int32 TotalBytesDownloaded, const float DownloadDuration);
+	void OnDownloadMetrics(const uint64 DownloadId, const int32 TotalBytesDownloaded, const float DownloadDuration, const int64_t DownloadStartDate, const int64_t DownloadEndDate);
 };
diff --git a/Engine/Source/Runtime/Online/BackgroundHTTP/Public/Interfaces/IBackgroundHttpRequest.h b/Engine/Source/Runtime/Online/BackgroundHTTP/Public/Interfaces/IBackgroundHttpRequest.h
--- a/Engine/Source/Runtime/Online/BackgroundHTTP/Public/Interfaces/IBackgroundHttpRequest.h
+++ b/Engine/Source/Runtime/Online/BackgroundHTTP/Public/Interfaces/IBackgroundHttpRequest.h
@@ -35,8 +35,10 @@ DECLARE_DELEGATE_ThreeParams(FBackgroundHttpProgressUpdateDelegate, FBackgroundH
  * @param Request Background Http Request for which the metrics were calculated
  * @param TotalBytesDownloaded The count of bytes downloaded in this request
  * @param DownloadDuration The duration of the download in seconds
+ * @param FetchStartTimeUTC The start time of the download
+ * @param FetchEndTimeUTC The end time of the download
  */
-DECLARE_DELEGATE_ThreeParams(FBackgroundHttpRequestMetricsDelegate, FBackgroundHttpRequestPtr /* Request */, int32 /* TotalBytesDownloaded */, float /* DownloadDuration */);
+DECLARE_DELEGATE_FiveParams(FBackgroundHttpRequestMetricsDelegate, FBackgroundHttpRequestPtr /* Request */, int32 /* TotalBytesDownloaded */, float /* DownloadDuration */, FDateTime /* FetchStartTimeUTC*/, FDateTime /* FetchEndTimeUTC */);
 
 /**
  * Enum used to describe download priority. Higher priorities will be downloaded first.
@@ -169,7 +171,7 @@ public:
 	/**
 	 * Set metrics for analytics event
 	 */
-	virtual void SetMetrics(const int32 TotalBytesDownloaded, const float DownloadDuration) {}
+	virtual void SetMetrics(const int32 TotalBytesDownloaded, const float DownloadDuration, const FDateTime& FetchStartTimeUTC, const FDateTime& FetchEndTimeUTC) {}
 
 	/**
 	* Called to begin processing the request.
