Add FetchStartTimeUTC and FetchEndTimeUTC to ChunkDownloadComplete

#was-backout 44495359
#rnx

diff --git a/Engine/Plugins/Online/Android/AndroidFetchBackgroundDownload/Source/AndroidFetchBackgroundDownload/Java/download/DownloadProgressListener.java b/Engine/Plugins/Online/Android/AndroidFetchBackgroundDownload/Source/AndroidFetchBackgroundDownload/Java/download/DownloadProgressListener.java
--- a/Engine/Plugins/Online/Android/AndroidFetchBackgroundDownload/Source/AndroidFetchBackgroundDownload/Java/download/DownloadProgressListener.java
+++ b/Engine/Plugins/Online/Android/AndroidFetchBackgroundDownload/Source/AndroidFetchBackgroundDownload/Java/download/DownloadProgressListener.java
@@ -9,7 +9,7 @@ public interface DownloadProgressListener
 	public void OnDownloadProgress(String RequestID, long BytesWrittenSinceLastCall, long TotalBytesWritten);
 	public void OnDownloadGroupProgress(int GroupID, int Progress, boolean Indeterminate);
 	public void OnDownloadComplete(String RequestID, String CompleteLocation, EDownloadCompleteReason CompleteReason);
-	public void OnDownloadMetrics(String RequestID, long TotalBytesDownloaded, long DownloadDuration);
+	public void OnDownloadMetrics(String RequestID, long TotalBytesDownloaded, long DownloadDuration, long DownloadStartTimeUTC, long DownloadEndTimeUTC);
 	public void OnAllDownloadsComplete(boolean bDidAllRequestsSucceed);
 	public void OnDownloadEnqueued(String RequestID, boolean bEnqueueSuccess);
 }
\ No newline at end of file
diff --git a/Engine/Plugins/Online/Android/AndroidFetchBackgroundDownload/Source/AndroidFetchBackgroundDownload/Java/download/UEDownloadWorker.java b/Engine/Plugins/Online/Android/AndroidFetchBackgroundDownload/Source/AndroidFetchBackgroundDownload/Java/download/UEDownloadWorker.java
--- a/Engine/Plugins/Online/Android/AndroidFetchBackgroundDownload/Source/AndroidFetchBackgroundDownload/Java/download/UEDownloadWorker.java
+++ b/Engine/Plugins/Online/Android/AndroidFetchBackgroundDownload/Source/AndroidFetchBackgroundDownload/Java/download/UEDownloadWorker.java
@@ -810,11 +810,11 @@ public class UEDownloadWorker extends UEWorker implements DownloadProgressListen
 	}
 
 	@Override
-	public void OnDownloadMetrics(String RequestID, long TotalBytesDownloaded, long DownloadDuration)
+	public void OnDownloadMetrics(String RequestID, long TotalBytesDownloaded, long DownloadDuration, long DownloadStartTimeUTC, long DownloadEndTimeUTC)
 	{
 		try
 		{
-			nativeAndroidBackgroundDownloadOnMetrics(RequestID, TotalBytesDownloaded, DownloadDuration);
+			nativeAndroidBackgroundDownloadOnMetrics(RequestID, TotalBytesDownloaded, DownloadDuration, DownloadStartTimeUTC, DownloadEndTimeUTC);
 		}
 		catch (UnsatisfiedLinkError e)
 		{
@@ -957,7 +957,7 @@ public class UEDownloadWorker extends UEWorker implements DownloadProgressListen
 	native void nativeAndroidBackgroundDownloadOnWorkerStop(String WorkID);
 	native void nativeAndroidBackgroundDownloadOnProgress(String TaskID, long BytesWrittenSinceLastCall, long TotalBytesWritten);
 	native void nativeAndroidBackgroundDownloadOnComplete(String TaskID, String CompleteLocation, boolean bWasSuccess);
-	native void nativeAndroidBackgroundDownloadOnMetrics(String TaskID, long TotalBytesDownloaded, long DownloadDuration);
+	native void nativeAndroidBackgroundDownloadOnMetrics(String TaskID, long TotalBytesDownloaded, long DownloadDuration, long DownloadStartTimeUTC, long DownloadEndTimeUTC);
 	native void nativeAndroidBackgroundDownloadOnAllComplete(boolean bDidAllRequestsSucceed);
 	native void nativeAndroidBackgroundDownloadOnTick();
 	
diff --git a/Engine/Plugins/Online/Android/AndroidFetchBackgroundDownload/Source/AndroidFetchBackgroundDownload/Java/download/fetch/FetchManager.java b/Engine/Plugins/Online/Android/AndroidFetchBackgroundDownload/Source/AndroidFetchBackgroundDownload/Java/download/fetch/FetchManager.java
--- a/Engine/Plugins/Online/Android/AndroidFetchBackgroundDownload/Source/AndroidFetchBackgroundDownload/Java/download/fetch/FetchManager.java
+++ b/Engine/Plugins/Online/Android/AndroidFetchBackgroundDownload/Source/AndroidFetchBackgroundDownload/Java/download/fetch/FetchManager.java
@@ -1271,14 +1271,16 @@ public class FetchManager implements FetchDownloadProgressOwner, FetchEnqueueRes
 		}
 
 		long DownloadDuration = -1;
+		long DownloadFinishTime = System.currentTimeMillis();
+
 		if (DownloadDesc.DownloadStartTime > 0)
 		{
-			long ElapsedDownloadTime = System.currentTimeMillis() - DownloadDesc.DownloadStartTime;
+			long ElapsedDownloadTime = DownloadFinishTime - DownloadDesc.DownloadStartTime;
 
 			DownloadDuration = ElapsedDownloadTime - DownloadDesc.TotalPausedTime;
 		}
 
-		DownloadDesc.ProgressListener.OnDownloadMetrics(DownloadDesc.RequestID, DownloadDesc.TotalDownloadedBytes, DownloadDuration);
+		DownloadDesc.ProgressListener.OnDownloadMetrics(DownloadDesc.RequestID, DownloadDesc.TotalDownloadedBytes, DownloadDuration, DownloadDesc.DownloadStartTime, DownloadFinishTime);
 
 		DownloadDesc.ProgressListener.OnDownloadComplete(DownloadDesc.RequestID, DownloadDesc.DestinationLocation, CompleteReasonToSend);
 		
diff --git a/Engine/Plugins/Online/Android/AndroidFetchBackgroundDownload/Source/AndroidFetchBackgroundDownload/Private/AndroidPlatformBackgroundHttpManager.cpp b/Engine/Plugins/Online/Android/AndroidFetchBackgroundDownload/Source/AndroidFetchBackgroundDownload/Private/AndroidPlatformBackgroundHttpManager.cpp
--- a/Engine/Plugins/Online/Android/AndroidFetchBackgroundDownload/Source/AndroidFetchBackgroundDownload/Private/AndroidPlatformBackgroundHttpManager.cpp
+++ b/Engine/Plugins/Online/Android/AndroidFetchBackgroundDownload/Source/AndroidFetchBackgroundDownload/Private/AndroidPlatformBackgroundHttpManager.cpp
@@ -58,7 +58,7 @@ namespace FAndroidBackgroundDownloadDelegates
 	DECLARE_DELEGATE_TwoParams(FAndroidBackgroundDownload_OnWorkerStop, FString /*WorkID*/, jobject /*UnderlyingWorker*/);
 	DECLARE_DELEGATE_FourParams(FAndroidBackgroundDownload_OnProgress, jobject /*UnderlyingWorker*/, FString /*RequestID*/, int64_t /*BytesWrittenSinceLastCall*/, int64_t /*TotalBytesWritten*/);
 	DECLARE_DELEGATE_FourParams(FAndroidBackgroundDownload_OnComplete, jobject /*UnderlyingWorker*/, FString /*RequestID*/, FString /*CompleteLocation*/, bool /*bWasSuccess*/);
-	DECLARE_DELEGATE_FourParams(FAndroidBackgroundDownload_OnMetrics, jobject /*UnderlyingWorker*/, FString /*RequestID*/, int64_t /*TotalBytesDownloaded*/, int64_t /*DownloadDuration*/);
+	DECLARE_DELEGATE_SixParams(FAndroidBackgroundDownload_OnMetrics, jobject /*UnderlyingWorker*/, FString /*RequestID*/, int64_t /*TotalBytesDownloaded*/, int64_t /*DownloadDuration*/, int64_t /*DownloadStartTimeUTC*/, int64_t /*DownloadEndTimeUTC*/);
 	DECLARE_DELEGATE_TwoParams(FAndroidBackgroundDownload_OnAllComplete, jobject /*UnderlyingWorker*/, bool /*bDidAllRequestsSucceed*/);
 	DECLARE_DELEGATE_TwoParams(FAndroidBackgroundDownload_OnTickWorkerThread, JNIEnv*, jobject /*UnderlyingWorker*/);
 
@@ -696,16 +696,21 @@ void FAndroidPlatformBackgroundHttpManager::Java_OnDownloadComplete(jobject Unde
 	}
 }
 
-void FAndroidPlatformBackgroundHttpManager::Java_OnDownloadMetrics(jobject UnderlyingWorker, FString RequestID, int64_t TotalBytesDownloaded, int64_t DownloadDuration)
+void FAndroidPlatformBackgroundHttpManager::Java_OnDownloadMetrics(jobject UnderlyingWorker, FString RequestID, int64_t TotalBytesDownloaded, int64_t DownloadDuration, int64_t DownloadStartTimeUTC, int64_t DownloadEndTimeUTC)
 {
 	UE_LOG(LogBackgroundHttpManager, Log, TEXT("DownloadMetrics... RequestID:%s TotalBytesDownloaded:%ld DownloadDuration:%ld"), *RequestID, TotalBytesDownloaded, DownloadDuration);
 
 	FAndroidBackgroundHttpRequestPtr CompletedRequest = FindRequestByID(RequestID);
 	if (CompletedRequest.IsValid())
 	{
-		float DownloadDurationInSeconds = (float)((double)DownloadDuration / 1000.0f);
+		FBackgroundHttpRequestMetricsExtended MetricsExtended;
 
-		CompletedRequest->NotifyRequestMetricsAvailable((int32)TotalBytesDownloaded, DownloadDurationInSeconds);
+		MetricsExtended.TotalBytesDownloaded = TotalBytesDownloaded;
+		MetricsExtended.DownloadDuration = (float)((double)DownloadDuration / 1000.0f);
+		MetricsExtended.FetchStartTimeUTC = FDateTime::FromUnixTimestampDecimal((double)DownloadStartTimeUTC / 1000.0f);
+		MetricsExtended.FetchEndTimeUTC = FDateTime::FromUnixTimestampDecimal((double)DownloadEndTimeUTC / 1000.0f);
+
+		CompletedRequest->NotifyRequestMetricsExtendedAvailable(MetricsExtended);
 	}
 }
 
@@ -1256,13 +1261,15 @@ JNI_METHOD void Java_com_epicgames_unreal_download_UEDownloadWorker_nativeAndroi
 	FAndroidBackgroundDownloadDelegates::AndroidBackgroundDownload_OnComplete.ExecuteIfBound(thiz, RequestID, ConvertedCompleteLocation, ConvertedbWasSuccess);
 }
 
-JNI_METHOD void Java_com_epicgames_unreal_download_UEDownloadWorker_nativeAndroidBackgroundDownloadOnMetrics(JNIEnv* jenv, jobject thiz, jstring TaskID, jlong TotalBytesDownloaded, jlong DownloadDuration)
+JNI_METHOD void Java_com_epicgames_unreal_download_UEDownloadWorker_nativeAndroidBackgroundDownloadOnMetrics(JNIEnv* jenv, jobject thiz, jstring TaskID, jlong TotalBytesDownloaded, jlong DownloadDuration, jlong DownloadStartTimeUTC, jlong DownloadEndTimeUTC)
 {
 	FString RequestID = FJavaHelper::FStringFromParam(jenv, TaskID);
 	int64_t ConvertedTotalBytesDownloaded = static_cast<uint64_t>(TotalBytesDownloaded);
 	int64_t ConvertedDownloadDuration = static_cast<uint64_t>(DownloadDuration);
+	int64_t ConvertedDownloadStartTimeUTC = static_cast<uint64_t>(DownloadStartTimeUTC);
+	int64_t ConvertedDownloadEndTimeUTC = static_cast<uint64_t>(DownloadEndTimeUTC);
 
-	FAndroidBackgroundDownloadDelegates::AndroidBackgroundDownload_OnMetrics.ExecuteIfBound(thiz, RequestID, ConvertedTotalBytesDownloaded, ConvertedDownloadDuration);
+	FAndroidBackgroundDownloadDelegates::AndroidBackgroundDownload_OnMetrics.ExecuteIfBound(thiz, RequestID, ConvertedTotalBytesDownloaded, ConvertedDownloadDuration, ConvertedDownloadStartTimeUTC, ConvertedDownloadEndTimeUTC);
 }
 
 JNI_METHOD void Java_com_epicgames_unreal_download_UEDownloadWorker_nativeAndroidBackgroundDownloadOnAllComplete(JNIEnv* jenv, jobject thiz, jboolean bDidAllRequestsSucceed)
diff --git a/Engine/Plugins/Online/Android/AndroidFetchBackgroundDownload/Source/AndroidFetchBackgroundDownload/Public/AndroidPlatformBackgroundHttpManager.h b/Engine/Plugins/Online/Android/AndroidFetchBackgroundDownload/Source/AndroidFetchBackgroundDownload/Public/AndroidPlatformBackgroundHttpManager.h
--- a/Engine/Plugins/Online/Android/AndroidFetchBackgroundDownload/Source/AndroidFetchBackgroundDownload/Public/AndroidPlatformBackgroundHttpManager.h
+++ b/Engine/Plugins/Online/Android/AndroidFetchBackgroundDownload/Source/AndroidFetchBackgroundDownload/Public/AndroidPlatformBackgroundHttpManager.h
@@ -74,7 +74,7 @@ protected:
 	void Java_OnWorkerStop(FString WorkID, jobject UnderlyingWorker);
 	void Java_OnDownloadProgress(jobject UnderlyingWorker, FString RequestID, int64_t BytesWrittenSinceLastCall, int64_t TotalBytesWritten);
 	void Java_OnDownloadComplete(jobject UnderlyingWorker, FString RequestID, FString CompleteLocation, bool bWasSuccess);
-	void Java_OnDownloadMetrics(jobject UnderlyingWorker, FString RequestID, int64_t TotalBytesDownloaded, int64_t DownloadDuration);
+	void Java_OnDownloadMetrics(jobject UnderlyingWorker, FString RequestID, int64_t TotalBytesDownloaded, int64_t DownloadDuration, int64_t DownloadStartTimeUTC, int64_t DownloadEndTimeUTC);
 	void Java_OnAllDownloadsComplete(jobject UnderlyingWorker, bool bDidAllRequestsSucceed);
 	void Java_OnTick(JNIEnv* Env, jobject UnderlyingWorker);
 
diff --git a/Engine/Source/Runtime/ApplicationCore/ApplicationCore.Build.cs b/Engine/Source/Runtime/ApplicationCore/ApplicationCore.Build.cs
--- a/Engine/Source/Runtime/ApplicationCore/ApplicationCore.Build.cs
+++ b/Engine/Source/Runtime/ApplicationCore/ApplicationCore.Build.cs
@@ -83,6 +83,7 @@ public class ApplicationCore : ModuleRules
 			//Need to add this as BackgroundHTTP files can end up doing work directly from our AppDelegate in iOS and thus we need acccess to correct file locations to save these very early.
 			PrivateDependencyModuleNames.Add("BackgroundHTTPFileHash");
 			PrivateDependencyModuleNames.Add("CoreUObject");
+			PublicIncludePathModuleNames.Add("BackgroundHTTP");
 		}
 		else if (Target.Platform == UnrealTargetPlatform.Android)
 		{
diff --git a/Engine/Source/Runtime/ApplicationCore/Private/IOS/IOSBackgroundURLSessionHandler.cpp b/Engine/Source/Runtime/ApplicationCore/Private/IOS/IOSBackgroundURLSessionHandler.cpp
--- a/Engine/Source/Runtime/ApplicationCore/Private/IOS/IOSBackgroundURLSessionHandler.cpp
+++ b/Engine/Source/Runtime/ApplicationCore/Private/IOS/IOSBackgroundURLSessionHandler.cpp
@@ -1026,18 +1026,26 @@ static constexpr NSInteger HTTPStatusCodeErrorServer = 500;
 #if !UE_BUILD_SHIPPING
             UE_DNLD_LOG(@"Response Duration: %f seconds", responseDuration);
 #endif
-
 			// Calculate and log download speed
-            if (responseDuration > 0 && metric.countOfResponseBodyBytesReceived > 0) {
-				
-				FBackgroundURLSessionHandler::OnDownloadMetrics.Broadcast((uint64)task.taskIdentifier, (int32)metric.countOfResponseBodyBytesReceived, (float)responseDuration);
+            if (responseDuration > 0 && metric.countOfResponseBodyBytesReceived > 0) 
+			{
+				FBackgroundHttpRequestMetricsExtended MetricsExtended;
+
+				MetricsExtended.TotalBytesDownloaded = (int32)metric.countOfResponseBodyBytesReceived;
+				MetricsExtended.DownloadDuration = (float)responseDuration;
+				MetricsExtended.FetchStartTimeUTC = FDateTime::FromUnixTimestampDecimal([metric.requestStartDate timeIntervalSince1970]);
+				MetricsExtended.FetchEndTimeUTC = FDateTime::FromUnixTimestampDecimal([metric.responseEndDate timeIntervalSince1970]);
+
+				FBackgroundURLSessionHandler::OnDownloadMetricsExtended.Broadcast((uint64)task.taskIdentifier, MetricsExtended);
 				
 #if !UE_BUILD_SHIPPING
 				// bytes per second
 				double downloadSpeed = (double)metric.countOfResponseBodyBytesReceived / responseDuration;
 				NSString *formattedSpeed = [self formattedSpeed:downloadSpeed];
 				UE_DNLD_LOG(@"Download Speed: %@", formattedSpeed);
-            } else {
+            } 
+			else 
+			{
                 UE_DNLD_LOG(@"Download Speed: Not Available");
 #endif
             }
@@ -1644,7 +1652,7 @@ static constexpr NSInteger HTTPStatusCodeErrorServer = 500;
 const uint64 FBackgroundURLSessionHandler::InvalidDownloadId = [FBackgroundNSURLSession GetInvalidDownloadId];
 
 FBackgroundURLSessionHandler::FOnDownloadCompleted FBackgroundURLSessionHandler::OnDownloadCompleted;
-FBackgroundURLSessionHandler::FOnDownloadMetrics FBackgroundURLSessionHandler::OnDownloadMetrics;
+FBackgroundURLSessionHandler::FOnDownloadMetricsExtended FBackgroundURLSessionHandler::OnDownloadMetricsExtended;
 
 FBackgroundURLSessionHandler::FOnDownloadsCompletedWhileAppWasNotRunning FBackgroundURLSessionHandler::OnDownloadsCompletedWhileAppWasNotRunning;
 
diff --git a/Engine/Source/Runtime/ApplicationCore/Public/IOS/IOSBackgroundURLSessionHandler.h b/Engine/Source/Runtime/ApplicationCore/Public/IOS/IOSBackgroundURLSessionHandler.h
--- a/Engine/Source/Runtime/ApplicationCore/Public/IOS/IOSBackgroundURLSessionHandler.h
+++ b/Engine/Source/Runtime/ApplicationCore/Public/IOS/IOSBackgroundURLSessionHandler.h
@@ -7,6 +7,9 @@
 
 #include "BackgroundHttpFileHashHelper.h"
 
+//Included for FBackgroundHttpRequestMetricsExtended
+#include "BackgroundHttpMetrics.h"
+
 //Call backs called by the bellow FBackgroundURLSessionHandler so higher-level systems can respond to task updates.
 class UE_DEPRECATED(5.5, "Use new download methods in FBackgroundURLSessionHandler.") APPLICATIONCORE_API FIOSBackgroundDownloadCoreDelegates
 {
@@ -61,9 +64,15 @@ public:
 	static FOnDownloadCompleted OnDownloadCompleted;
 	
 	// Will be invoked from didFinishCollectingMetrics
+	UE_DEPRECATED(5.7, "FBackgroundURLSessionHandler::FOnDownloadMetrics is deprecated.  Use FOnDownloadMetricsExtended")
 	DECLARE_MULTICAST_DELEGATE_ThreeParams(FOnDownloadMetrics, const uint64 /*DownloadId*/, const int32 /*TotalBytesDownloaded*/, const float /*DownloadDuration*/)
+	UE_DEPRECATED(5.7, "FBackgroundURLSessionHandler::OnDownloadMetrics is deprecated.  Use OnDownloadMetricsExtended")
 	static FOnDownloadMetrics OnDownloadMetrics;
 
+	// Will be invoked from didFinishCollectingMetrics
+	DECLARE_MULTICAST_DELEGATE_TwoParams(FOnDownloadMetricsExtended, const uint64 /*DownloadId*/, const FBackgroundHttpRequestMetricsExtended /*MetricsExtended*/);
+	static FOnDownloadMetricsExtended OnDownloadMetricsExtended;
+
 	// Will be invoked from handleEventsForBackgroundURLSession application delegate. Needs to be registered very early, e.g. from static constructor.
 	// handleEventsForBackgroundURLSession is only invoked if app was killed by OS while in background and then relaunched to notify that downloads were completed.
 	// Is not invoked in any other scenario.
diff --git a/Engine/Source/Runtime/InstallBundleManager/Public/InstallBundleManagerInterface.h b/Engine/Source/Runtime/InstallBundleManager/Public/InstallBundleManagerInterface.h
--- a/Engine/Source/Runtime/InstallBundleManager/Public/InstallBundleManagerInterface.h
+++ b/Engine/Source/Runtime/InstallBundleManager/Public/InstallBundleManagerInterface.h
@@ -96,6 +96,8 @@ struct FInstallBundleChunkDownloadInfo
 	FString ChunkName;
 	int32 ChunkFileSize;
 	float ChunkDownloadDuration;
+	FDateTime FetchStartTimeUTC;
+	FDateTime FetchEndTimeUTC;
 };
 
 using FInstallBundleSourceOrCache = TUnion<FInstallBundleSourceType, FName>;
diff --git a/Engine/Source/Runtime/Online/BackgroundHTTP/Private/BackgroundHttpManagerImpl.cpp b/Engine/Source/Runtime/Online/BackgroundHTTP/Private/BackgroundHttpManagerImpl.cpp
--- a/Engine/Source/Runtime/Online/BackgroundHTTP/Private/BackgroundHttpManagerImpl.cpp
+++ b/Engine/Source/Runtime/Online/BackgroundHTTP/Private/BackgroundHttpManagerImpl.cpp
@@ -364,8 +364,15 @@ bool FBackgroundHttpManagerImpl::AssociateWithAnyExistingRequest(const FBackgrou
 			//First send progress update for the file size so anything monitoring this download knows we are about to update this progress
 			Request->OnProgressUpdated().ExecuteIfBound(Request, ExistingFileSize, ExistingFileSize);
 
-			//Now complete with this completed response data
-			Request->SetMetrics(ExistingFileSize, 0.0f);
+			//Now complete with this response data.  Since we have no time data fill everything in with zeros
+			FBackgroundHttpRequestMetricsExtended MetricsExtended;
+
+			MetricsExtended.TotalBytesDownloaded = ExistingFileSize;
+			MetricsExtended.DownloadDuration = 0.0f;
+			MetricsExtended.FetchStartTimeUTC = FDateTime::FromUnixTimestamp(0);
+			MetricsExtended.FetchEndTimeUTC = FDateTime::FromUnixTimestamp(0);
+
+			Request->SetMetricsExtended(MetricsExtended);
 			Request->CompleteWithExistingResponseData(NewResponseWithExistingFile);
 		}
 	}
diff --git a/Engine/Source/Runtime/Online/BackgroundHTTP/Private/BackgroundHttpRequestImpl.cpp b/Engine/Source/Runtime/Online/BackgroundHTTP/Private/BackgroundHttpRequestImpl.cpp
--- a/Engine/Source/Runtime/Online/BackgroundHTTP/Private/BackgroundHttpRequestImpl.cpp
+++ b/Engine/Source/Runtime/Online/BackgroundHTTP/Private/BackgroundHttpRequestImpl.cpp
@@ -76,12 +76,16 @@ void FBackgroundHttpRequestImpl::OnBackgroundDownloadComplete()
 	
 	if (OptionalMetricsInfo.IsSet())
 	{
-		if (!OnRequestMetrics().IsBound())
+		PRAGMA_DISABLE_DEPRECATION_WARNINGS
+		if (!OnRequestMetrics().IsBound() && !OnRequestMetricsExtended().IsBound())
 		{
-			UE_LOG(LogBackgroundHttpRequest, Warning, TEXT("OptionalMetricsInfo OnRequestMetrics() is NOT bound - RequestID:%s"), *GetRequestID());
+			UE_LOG(LogBackgroundHttpRequest, Warning, TEXT("OptionalMetricsInfo neither OnRequestMetrics() or OnRequestMetricsExtended() are bound - RequestID:%s"), *GetRequestID());
 		}
+
 		OnRequestMetrics().ExecuteIfBound(SharedThis(this), OptionalMetricsInfo->TotalBytesDownloaded, OptionalMetricsInfo->DownloadDuration);
+		OnRequestMetricsExtended().ExecuteIfBound(SharedThis(this), OptionalMetricsInfo.GetValue());
 		OptionalMetricsInfo.Reset();
+		PRAGMA_ENABLE_DEPRECATION_WARNINGS
 	}
 	else
 	{
@@ -105,7 +109,9 @@ void FBackgroundHttpRequestImpl::NotifyNotificationObjectOfComplete(bool bWasSuc
 
 void FBackgroundHttpRequestImpl::SetMetrics(const int32 TotalBytesDownloaded, const float DownloadDuration)
 {
+	PRAGMA_DISABLE_DEPRECATION_WARNINGS
 	NotifyRequestMetricsAvailable(TotalBytesDownloaded, DownloadDuration);
+	PRAGMA_ENABLE_DEPRECATION_WARNINGS
 }
 
 void FBackgroundHttpRequestImpl::NotifyRequestMetricsAvailable(const int32 TotalBytesDownloaded, const float DownloadDuration)
@@ -113,10 +119,22 @@ void FBackgroundHttpRequestImpl::NotifyRequestMetricsAvailable(const int32 Total
 	ensureAlwaysMsgf(!OptionalMetricsInfo.IsSet(), TEXT("Overwriting metrics data!!!"));
 	OptionalMetricsInfo = {
 		.TotalBytesDownloaded = TotalBytesDownloaded,
-		.DownloadDuration = DownloadDuration
+		.DownloadDuration = DownloadDuration,
 	};
 }
 
+void FBackgroundHttpRequestImpl::SetMetricsExtended(const FBackgroundHttpRequestMetricsExtended ExtendedMetrics)
+{
+	NotifyRequestMetricsExtendedAvailable(ExtendedMetrics);
+}
+
+void FBackgroundHttpRequestImpl::NotifyRequestMetricsExtendedAvailable(const FBackgroundHttpRequestMetricsExtended ExtendedMetrics)
+{
+	ensureAlwaysMsgf(!OptionalMetricsInfo.IsSet(), TEXT("Overwriting metrics data!!!"));
+	OptionalMetricsInfo = ExtendedMetrics;
+}
+
+
 void FBackgroundHttpRequestImpl::SetURLAsList(const TArray<FString>& URLs, int NumRetriesIn)
 {
 	NumberOfTotalRetries = NumRetriesIn;
@@ -173,6 +191,11 @@ FBackgroundHttpRequestMetricsDelegate& FBackgroundHttpRequestImpl::OnRequestMetr
 	return HttpRequestMetricsDelegate;
 }
 
+FBackgroundHttpRequestMetricsExtendedDelegate& FBackgroundHttpRequestImpl::OnRequestMetricsExtended()
+{
+	return HttpRequestMetricsExtendedDelegate;
+}
+
 const FBackgroundHttpResponsePtr FBackgroundHttpRequestImpl::GetResponse() const
 {
 	return Response;
diff --git a/Engine/Source/Runtime/Online/BackgroundHTTP/Private/IOS/ApplePlatformBackgroundHttpManager.cpp b/Engine/Source/Runtime/Online/BackgroundHTTP/Private/IOS/ApplePlatformBackgroundHttpManager.cpp
--- a/Engine/Source/Runtime/Online/BackgroundHTTP/Private/IOS/ApplePlatformBackgroundHttpManager.cpp
+++ b/Engine/Source/Runtime/Online/BackgroundHTTP/Private/IOS/ApplePlatformBackgroundHttpManager.cpp
@@ -24,13 +24,13 @@ FApplePlatformBackgroundHttpManager::FApplePlatformBackgroundHttpManager()
 	, PendingRemoveRequestLock()
 {
 	OnDownloadCompletedHandle = FBackgroundURLSessionHandler::OnDownloadCompleted.AddRaw(this, &FApplePlatformBackgroundHttpManager::OnDownloadCompleted);
-	OnDownloadMetricsHandle = FBackgroundURLSessionHandler::OnDownloadMetrics.AddRaw(this, &FApplePlatformBackgroundHttpManager::OnDownloadMetrics);
+	OnDownloadMetricsExtendedHandle = FBackgroundURLSessionHandler::OnDownloadMetricsExtended.AddRaw(this, &FApplePlatformBackgroundHttpManager::OnDownloadMetricsExtended);
 }
 
 FApplePlatformBackgroundHttpManager::~FApplePlatformBackgroundHttpManager()
 {
 	FBackgroundURLSessionHandler::OnDownloadCompleted.Remove(OnDownloadCompletedHandle);
-	FBackgroundURLSessionHandler::OnDownloadMetrics.Remove(OnDownloadMetricsHandle);
+	FBackgroundURLSessionHandler::OnDownloadMetricsExtended.Remove(OnDownloadMetricsExtendedHandle);
 }
 
 void FApplePlatformBackgroundHttpManager::AddRequest(const FBackgroundHttpRequestPtr GenericRequest)
@@ -138,7 +138,7 @@ void FApplePlatformBackgroundHttpManager::OnDownloadCompleted(const uint64 Downl
 	}
 }
 
-void FApplePlatformBackgroundHttpManager::OnDownloadMetrics(const uint64 DownloadId, const int32 TotalBytesDownloaded, const float DownloadDuration)
+void FApplePlatformBackgroundHttpManager::OnDownloadMetricsExtended(const uint64 DownloadId, const FBackgroundHttpRequestMetricsExtended MetricsExtended)
 {
 	FRWScopeLock ScopeLock(ActiveRequestLock, SLT_ReadOnly);
 	for (FBackgroundHttpRequestPtr& GenericRequest : ActiveRequests)
@@ -148,10 +148,10 @@ void FApplePlatformBackgroundHttpManager::OnDownloadMetrics(const uint64 Downloa
 		{
 			if (Request->GetInternalDownloadId() == DownloadId)
 			{
-				Request->NotifyRequestMetricsAvailable(TotalBytesDownloaded, DownloadDuration);
+				Request->NotifyRequestMetricsExtendedAvailable(MetricsExtended);
 				return;
 			}
 		}
 	}
-	UE_LOG(LogBackgroundHttpManager, Warning, TEXT("Fail to NotifyRequestMetricsAvailable - RequestID:%llu | TotalBytesDownloaded:%d | DownloadDuration:%f"), DownloadId, TotalBytesDownloaded, DownloadDuration);
+	UE_LOG(LogBackgroundHttpManager, Warning, TEXT("Fail to NotifyRequestMetricsExtendedAvailable - RequestID:%llu | TotalBytesDownloaded:%d | DownloadDuration:%f"), DownloadId, MetricsExtended.TotalBytesDownloaded, MetricsExtended.DownloadDuration);
 }
diff --git a/Engine/Source/Runtime/Online/BackgroundHTTP/Public/BackgroundHttpMetrics.h b/Engine/Source/Runtime/Online/BackgroundHTTP/Public/BackgroundHttpMetrics.h
new file mode 100644
--- /dev/null
+++ b/Engine/Source/Runtime/Online/BackgroundHTTP/Public/BackgroundHttpMetrics.h
@@ -0,0 +1,13 @@
+// Copyright Epic Games, Inc. All Rights Reserved.
+
+#pragma once
+
+#include "CoreMinimal.h"
+
+struct FBackgroundHttpRequestMetricsExtended
+{
+	int32 TotalBytesDownloaded;
+	float DownloadDuration;
+	FDateTime FetchStartTimeUTC;
+	FDateTime FetchEndTimeUTC;
+};
diff --git a/Engine/Source/Runtime/Online/BackgroundHTTP/Public/BackgroundHttpRequestImpl.h b/Engine/Source/Runtime/Online/BackgroundHTTP/Public/BackgroundHttpRequestImpl.h
--- a/Engine/Source/Runtime/Online/BackgroundHTTP/Public/BackgroundHttpRequestImpl.h
+++ b/Engine/Source/Runtime/Online/BackgroundHTTP/Public/BackgroundHttpRequestImpl.h
@@ -38,7 +38,9 @@ public:
 	BACKGROUNDHTTP_API virtual void CompleteWithExistingResponseData(FBackgroundHttpResponsePtr BackgroundResponse) override;
 	BACKGROUNDHTTP_API virtual FBackgroundHttpRequestCompleteDelegate& OnProcessRequestComplete() override;
 	BACKGROUNDHTTP_API virtual FBackgroundHttpProgressUpdateDelegate& OnProgressUpdated() override;
+	UE_DEPRECATED(5.7, "FBackgroundHttpRequestImpl::OnRequestMetrics is deprecated, please use OnRequestMetricsExtended instead.")
 	BACKGROUNDHTTP_API virtual FBackgroundHttpRequestMetricsDelegate& OnRequestMetrics() override;
+	BACKGROUNDHTTP_API virtual FBackgroundHttpRequestMetricsExtendedDelegate& OnRequestMetricsExtended() override;
 	BACKGROUNDHTTP_API virtual const FBackgroundHttpResponsePtr GetResponse() const override;
 	BACKGROUNDHTTP_API virtual const FString& GetRequestID() const override;
 	BACKGROUNDHTTP_API virtual void SetRequestID(const FString& NewRequestID) override;
@@ -47,8 +49,14 @@ public:
 	BACKGROUNDHTTP_API virtual void SetRequestPriority(EBackgroundHTTPPriority NewPriority) override;
 	
 	BACKGROUNDHTTP_API virtual void NotifyNotificationObjectOfComplete(const bool bWasSuccess);
+
+	UE_DEPRECATED(5.7, "FBackgroundHttpRequestImpl::NotifyRequestMetricsAvailable is deprecated, please use NotifyRequestMetricsExtendedAvailable instead.")
 	BACKGROUNDHTTP_API virtual void NotifyRequestMetricsAvailable(const int32 TotalBytesDownloaded, const float DownloadDuration);
+	BACKGROUNDHTTP_API virtual void NotifyRequestMetricsExtendedAvailable(const FBackgroundHttpRequestMetricsExtended ExtendedMetrics);
+
+	UE_DEPRECATED(5.7, "FBackgroundHttpRequestImpl::SetMetrics is deprecated, please use SetMetricsExtended instead.")
 	BACKGROUNDHTTP_API virtual void SetMetrics(const int32 TotalBytesDownloaded, const float DownloadDuration) override;
+	BACKGROUNDHTTP_API virtual void SetMetricsExtended(const FBackgroundHttpRequestMetricsExtended ExtendedMetrics) override;
 	
 protected:
 	UE::FMutex DownloadCompleteMutex;
@@ -60,15 +68,12 @@ protected:
 	EBackgroundHTTPPriority RequestPriority;
 	uint64 ExpectedResultSize;
 	
-	struct FDownloadMetricsInfo
-	{
-		int32 TotalBytesDownloaded;
-		float DownloadDuration;
-	};
+	using FDownloadMetricsInfo = FBackgroundHttpRequestMetricsExtended;
 	
 	TOptional<FDownloadMetricsInfo> OptionalMetricsInfo;
 	
 	FBackgroundHttpRequestCompleteDelegate HttpRequestCompleteDelegate;
 	FBackgroundHttpProgressUpdateDelegate HttpProgressUpdateDelegate;
 	FBackgroundHttpRequestMetricsDelegate HttpRequestMetricsDelegate;
+	FBackgroundHttpRequestMetricsExtendedDelegate HttpRequestMetricsExtendedDelegate;
 };
diff --git a/Engine/Source/Runtime/Online/BackgroundHTTP/Public/IOS/ApplePlatformBackgroundHttpManager.h b/Engine/Source/Runtime/Online/BackgroundHTTP/Public/IOS/ApplePlatformBackgroundHttpManager.h
--- a/Engine/Source/Runtime/Online/BackgroundHTTP/Public/IOS/ApplePlatformBackgroundHttpManager.h
+++ b/Engine/Source/Runtime/Online/BackgroundHTTP/Public/IOS/ApplePlatformBackgroundHttpManager.h
@@ -32,7 +32,7 @@ private:
 	TArray<FBackgroundHttpRequestPtr> PendingRemoveRequests;
 	FRWLock PendingRemoveRequestLock;
 	FDelegateHandle OnDownloadCompletedHandle;
-	FDelegateHandle OnDownloadMetricsHandle;
+	FDelegateHandle OnDownloadMetricsExtendedHandle;
 	void OnDownloadCompleted(const uint64 DownloadId, const bool bSuccess);
-	void OnDownloadMetrics(const uint64 DownloadId, const int32 TotalBytesDownloaded, const float DownloadDuration);
+	void OnDownloadMetricsExtended(const uint64 DownloadId, const FBackgroundHttpRequestMetricsExtended MetricsExtended);
 };
diff --git a/Engine/Source/Runtime/Online/BackgroundHTTP/Public/Interfaces/IBackgroundHttpRequest.h b/Engine/Source/Runtime/Online/BackgroundHTTP/Public/Interfaces/IBackgroundHttpRequest.h
--- a/Engine/Source/Runtime/Online/BackgroundHTTP/Public/Interfaces/IBackgroundHttpRequest.h
+++ b/Engine/Source/Runtime/Online/BackgroundHTTP/Public/Interfaces/IBackgroundHttpRequest.h
@@ -8,11 +8,15 @@
 //Included as we use EHttpResponseCode
 #include "Interfaces/IHttpResponse.h"
 
+//Included for FBackgroundHttpRequestMetricsExtended
+#include "BackgroundHttpMetrics.h"
 
 typedef TSharedPtr<class IBackgroundHttpRequest, ESPMode::ThreadSafe> FBackgroundHttpRequestPtr;
 typedef TSharedPtr<class IBackgroundHttpResponse, ESPMode::ThreadSafe> FBackgroundHttpResponsePtr;
 typedef TSharedPtr<class IBackgroundHttpManager, ESPMode::ThreadSafe> FBackgroundHttpManagerPtr;
 
+
+
 /**
  * Delegate called when a Background Http request completes
  *
@@ -38,6 +42,13 @@ DECLARE_DELEGATE_ThreeParams(FBackgroundHttpProgressUpdateDelegate, FBackgroundH
  */
 DECLARE_DELEGATE_ThreeParams(FBackgroundHttpRequestMetricsDelegate, FBackgroundHttpRequestPtr /* Request */, int32 /* TotalBytesDownloaded */, float /* DownloadDuration */);
 
+/**
+ * Delegate called when a Background Http request extended metrics are ready
+ * @param Request Background Http Request for which the metrics were calculated
+ * @param BackgroundHttpRequestMetricsExtended struct containing all the metrics
+ */
+DECLARE_DELEGATE_TwoParams(FBackgroundHttpRequestMetricsExtendedDelegate, FBackgroundHttpRequestPtr /* Request */, FBackgroundHttpRequestMetricsExtended /* All the metrics on the request */);
+
 /**
  * Enum used to describe download priority. Higher priorities will be downloaded first.
  * Note: Should always be kept in High -> Low priority order if adding more Priorities!
@@ -164,13 +175,27 @@ public:
 	/**
 	 * Delegate called when the request's metrics were calculated
 	 */
+	UE_DEPRECATED(5.7, "IBackgroundHttpRequest::OnRequestMetrics is deprecated, please use IBackgroundHttpRequest::OnRequestMetricsExtended instead.")
 	virtual FBackgroundHttpRequestMetricsDelegate& OnRequestMetrics() = 0;
+
+	/**
+	 * Delegate called when the request's extended metrics were calculated
+	 */
+	virtual FBackgroundHttpRequestMetricsExtendedDelegate& OnRequestMetricsExtended() = 0;
+
 	
 	/**
 	 * Set metrics for analytics event
 	 */
+	UE_DEPRECATED(5.7, "IBackgroundHttpRequest::SetMetrics is deprecated, please use IBackgroundHttpRequest::SetMetricsExtended instead.")
 	virtual void SetMetrics(const int32 TotalBytesDownloaded, const float DownloadDuration) {}
 
+	/**
+	 * Set extended metrics for analytics event
+	 */
+	virtual void SetMetricsExtended(const FBackgroundHttpRequestMetricsExtended ExtendedMetrics) {}
+
+
 	/**
 	* Called to begin processing the request.
 	* OnProcessRequestComplete delegate is always called when the request completes or on error if it is bound.
