Delete old temp download files during startup on Android

#redirectsaresafe

diff --git a/Engine/Plugins/Online/Android/AndroidFetchBackgroundDownload/Source/AndroidFetchBackgroundDownload/AndroidFetchBackgroundDownload.Build.cs b/Engine/Plugins/Online/Android/AndroidFetchBackgroundDownload/Source/AndroidFetchBackgroundDownload/AndroidFetchBackgroundDownload.Build.cs
--- a/Engine/Plugins/Online/Android/AndroidFetchBackgroundDownload/Source/AndroidFetchBackgroundDownload/AndroidFetchBackgroundDownload.Build.cs
+++ b/Engine/Plugins/Online/Android/AndroidFetchBackgroundDownload/Source/AndroidFetchBackgroundDownload/AndroidFetchBackgroundDownload.Build.cs
@@ -18,6 +18,7 @@ namespace UnrealBuildTool.Rules
 					"Launch",
 					"AndroidBackgroundService",
 					"BackgroundHTTP",
+					"BuildSettings",
 				});
 
 			if (Target.Platform == UnrealTargetPlatform.Android)
diff --git a/Engine/Plugins/Online/Android/AndroidFetchBackgroundDownload/Source/AndroidFetchBackgroundDownload/AndroidFetchBackgroundDownload_UPL.xml b/Engine/Plugins/Online/Android/AndroidFetchBackgroundDownload/Source/AndroidFetchBackgroundDownload/AndroidFetchBackgroundDownload_UPL.xml
--- a/Engine/Plugins/Online/Android/AndroidFetchBackgroundDownload/Source/AndroidFetchBackgroundDownload/AndroidFetchBackgroundDownload_UPL.xml
+++ b/Engine/Plugins/Online/Android/AndroidFetchBackgroundDownload/Source/AndroidFetchBackgroundDownload/AndroidFetchBackgroundDownload_UPL.xml
@@ -69,4 +69,9 @@
 		<log text="Copying AndroidFetchBackgroundDownload files to staging"/>
 		<copyDir src="$S(PluginDir)/Java" dst="$S(BuildDir)/src/com/epicgames/unreal" />
 	</prebuildCopies>
+
+	<!-- set build version to BranchName-Changelist, see C++ equivalent in AndroidPlatformBackgroundHttpRequest.cpp-->
+	<gameActivityClassAdditions>
+		<insertValue value="public static final String BGDL_BUILD_VERSION = &quot;$S(EngineBranchName)&quot;;"/>
+	</gameActivityClassAdditions>
 </root>
diff --git a/Engine/Plugins/Online/Android/AndroidFetchBackgroundDownload/Source/AndroidFetchBackgroundDownload/Java/download/UEDownloadWorker.java b/Engine/Plugins/Online/Android/AndroidFetchBackgroundDownload/Source/AndroidFetchBackgroundDownload/Java/download/UEDownloadWorker.java
--- a/Engine/Plugins/Online/Android/AndroidFetchBackgroundDownload/Source/AndroidFetchBackgroundDownload/Java/download/UEDownloadWorker.java
+++ b/Engine/Plugins/Online/Android/AndroidFetchBackgroundDownload/Source/AndroidFetchBackgroundDownload/Java/download/UEDownloadWorker.java
@@ -132,13 +132,7 @@ public class UEDownloadWorker extends UEWorker implements DownloadProgressListen
 		
 		if (null == mFetchManager)
 		{
-			// set up fetch implementation, which might be unnnecessary if it was already done
-			Context context = getApplicationContext();
-			SharedPreferences preferences = context.getSharedPreferences("BackgroundPreferences", context.MODE_PRIVATE);
-			boolean bUseNewFetchImplementation = preferences.getBoolean("UseNewFetchImplementation", false);
-			FetchManagerInterfaceHelper.InitSharedManager(bUseNewFetchImplementation);
-
-			// get it
+			FetchManagerInterfaceHelper.InitSharedManager(getApplicationContext());
 			mFetchManager = FetchManagerInterfaceHelper.GetSharedManager();	
 		}
 		
diff --git a/Engine/Plugins/Online/Android/AndroidFetchBackgroundDownload/Source/AndroidFetchBackgroundDownload/Java/download/fetch/FetchManagerInterfaceHelper.java b/Engine/Plugins/Online/Android/AndroidFetchBackgroundDownload/Source/AndroidFetchBackgroundDownload/Java/download/fetch/FetchManagerInterfaceHelper.java
--- a/Engine/Plugins/Online/Android/AndroidFetchBackgroundDownload/Source/AndroidFetchBackgroundDownload/Java/download/fetch/FetchManagerInterfaceHelper.java
+++ b/Engine/Plugins/Online/Android/AndroidFetchBackgroundDownload/Source/AndroidFetchBackgroundDownload/Java/download/fetch/FetchManagerInterfaceHelper.java
@@ -3,13 +3,55 @@
 package com.epicgames.unreal.download.fetch;
 
 import com.epicgames.unreal.Logger;
+import com.epicgames.unreal.GameActivity;
+
+import android.content.Context;
+import android.content.SharedPreferences;
+
+import java.io.File;
 
 public class FetchManagerInterfaceHelper {
 
-	public static void InitSharedManager(boolean bUseNewImplementation)
+	public final static String NEW_FETCH_IMPLEMENATION_KEY = "UseNewFetchImplementation";
+	public final static String TEMP_DOWNLOAD_ROOT_PATH_KEY = "TempDownloadRootPath";
+
+	public static void InitSharedManager(Context context)
 	{
 		if (SharedManager == null) {
-			if (bUseNewImplementation) {
+			final SharedPreferences preferences = context.getSharedPreferences("BackgroundPreferences", context.MODE_PRIVATE);
+
+			// cleanup all temp content not in a directory named the same as our build version
+			String tempDownloadRootPath = null;
+			try {
+				tempDownloadRootPath = preferences.getString(TEMP_DOWNLOAD_ROOT_PATH_KEY, null);
+			} catch (ClassCastException notPresent) {}
+			if (tempDownloadRootPath != null && !tempDownloadRootPath.isEmpty()) {
+				File[] files = new File(tempDownloadRootPath).listFiles();
+				if (files != null) {
+					int deletedCount = 0;
+					for (File file : files) {
+						if (!file.getName().endsWith("DownloadDescriptionJSONs") &&
+							!file.getName().endsWith("URLMap") &&
+							!file.getName().equals(GameActivity.BGDL_BUILD_VERSION)) {
+							Log.debug("cleaning up: " + file.getName());
+							DeleteRecursive(file);
+							++deletedCount;
+						}
+					}
+					Log.debug("InitSharedManager: removed " + deletedCount + " files/directories from temp download directory: " + tempDownloadRootPath);
+				} else {
+					Log.debug("InitSharedManager: temp download directory did not exist: " + tempDownloadRootPath);
+				}
+			} else {
+				Log.debug("InitSharedManager: no known temp download directory to clean up");
+			}
+
+			// then decide what implementation to use, new or old, and instantiate it
+			boolean bUseNewFetchImplementation = false;
+			try {
+				bUseNewFetchImplementation = preferences.getBoolean(NEW_FETCH_IMPLEMENATION_KEY, false);
+			} catch (ClassCastException notPresent) {}
+			if (bUseNewFetchImplementation) {
 				Log.debug("Using new fetch manager implementation");
 				SharedManager = FetchManagerNew.GetSharedManager();
 			} else {
@@ -32,4 +74,13 @@ public class FetchManagerInterfaceHelper {
 
 	private static Logger Log = new Logger("UE", "FetchManagerInterfaceHelper");
 	private static FetchManagerInterface SharedManager = null;
+
+	public static boolean DeleteRecursive(File fileOrDirectory) {
+		if (fileOrDirectory.isDirectory()) {
+			for (File child : fileOrDirectory.listFiles()) {
+				DeleteRecursive(child);
+			}
+		}
+		return fileOrDirectory.delete();
+	}
 }
diff --git a/Engine/Plugins/Online/Android/AndroidFetchBackgroundDownload/Source/AndroidFetchBackgroundDownload/Java/download/fetch/FetchManagerNew.java b/Engine/Plugins/Online/Android/AndroidFetchBackgroundDownload/Source/AndroidFetchBackgroundDownload/Java/download/fetch/FetchManagerNew.java
--- a/Engine/Plugins/Online/Android/AndroidFetchBackgroundDownload/Source/AndroidFetchBackgroundDownload/Java/download/fetch/FetchManagerNew.java
+++ b/Engine/Plugins/Online/Android/AndroidFetchBackgroundDownload/Source/AndroidFetchBackgroundDownload/Java/download/fetch/FetchManagerNew.java
@@ -223,8 +223,6 @@ public class FetchManagerNew implements FetchManagerInterface
 
 		@Override
 		public void OnStart(String Id, URL url, long startingOffset, long totalBytesOrMinus1) {
-			Log.debug("OnStart: " + Id + " | " + url.toString() + " | " + startingOffset + " | " + totalBytesOrMinus1);
-
 			LastDownloadedSoFar = startingOffset;
 			TotalBytes = totalBytesOrMinus1;			// total target file size, regardless of starting offset
 
@@ -246,8 +244,6 @@ public class FetchManagerNew implements FetchManagerInterface
 		}
 		@Override
 		public void OnSuccess(String Id, File file) {
-			Log.debug("OnSuccess: " + Id + " | " + file.getName());
-
 			EDownloadCompleteReason Reason = EDownloadCompleteReason.Success;
 			File NewFileAtDestinationLocation = new File(TargetLocation);
 			if (!file.renameTo(NewFileAtDestinationLocation)) {
diff --git a/Engine/Plugins/Online/Android/AndroidFetchBackgroundDownload/Source/AndroidFetchBackgroundDownload/Private/AndroidPlatformBackgroundHttpManager.cpp b/Engine/Plugins/Online/Android/AndroidFetchBackgroundDownload/Source/AndroidFetchBackgroundDownload/Private/AndroidPlatformBackgroundHttpManager.cpp
--- a/Engine/Plugins/Online/Android/AndroidFetchBackgroundDownload/Source/AndroidFetchBackgroundDownload/Private/AndroidPlatformBackgroundHttpManager.cpp
+++ b/Engine/Plugins/Online/Android/AndroidFetchBackgroundDownload/Source/AndroidFetchBackgroundDownload/Private/AndroidPlatformBackgroundHttpManager.cpp
@@ -18,6 +18,7 @@
 #include "Android/AndroidJavaEnv.h"
 #include "Android/AndroidJNI.h"
 #include "Android/AndroidApplication.h"
+#include "Android/AndroidPlatformFile.h"
 
 #include "Containers/UnrealString.h"
 
@@ -25,6 +26,12 @@
 
 #include "HAL/IConsoleManager.h"
 
+#include "BuildSettings.h"
+
+#if USE_ANDROID_JNI
+extern void AndroidThunkCpp_SetSharedPreferenceString(const FString& Group, const FString& Key, const FString& Value);
+#endif
+
 #if UE_BUILD_SHIPPING
 // always clear any exceptions in shipping
 #define CHECK_JNI_RESULT(Id) if (Id == 0) { Env->ExceptionClear(); }
@@ -186,9 +193,22 @@ void FAndroidPlatformBackgroundHttpManager::Initialize()
 	FAndroidBackgroundDownloadDelegates::AndroidBackgroundDownload_OnMetrics.BindSP(this, &FAndroidPlatformBackgroundHttpManager::Java_OnDownloadMetrics);
 	FAndroidBackgroundDownloadDelegates::AndroidBackgroundDownload_OnAllComplete.BindSP(this, &FAndroidPlatformBackgroundHttpManager::Java_OnAllDownloadsComplete);
 	FAndroidBackgroundDownloadDelegates::AndroidBackgroundDownload_OnTickWorkerThread.BindSP(this, &FAndroidPlatformBackgroundHttpManager::Java_OnTick);
- 
+
+	AndroidThunkCpp_SetSharedPreferenceString(TEXT("BackgroundPreferences"), TEXT("TempDownloadRootPath"), FBackgroundHttpFileHashHelper().GetTemporaryRootPath());
+
 	AndroidBackgroundHTTPManagerDefaultLocalizedText.InitFromIniSettings("AndroidFetchBackgroundDownload");
 	FBackgroundHttpManagerImpl::Initialize();
+
+	// we download to our temporary root path, and then in a subdirectory named after our branch name;
+	// this way, the fetch managers can cleanup stuff they don't know (e.g. downloaded temp files from older versions)
+	bool bSuccess = false;
+	const FString BuildDownloadPath = FPaths::Combine(FBackgroundHttpFileHashHelper::GetTemporaryRootPath(), BuildSettings::GetBranchName());
+	if (IAndroidPlatformFile::GetPlatformPhysical().CreateDirectory(*FBackgroundHttpFileHashHelper::GetTemporaryRootPath())) {
+		bSuccess = IAndroidPlatformFile::GetPlatformPhysical().CreateDirectory(*BuildDownloadPath);
+	}
+	if (!bSuccess) {
+		UE_LOG(LogBackgroundHttpManager, Display, TEXT("Failed to create build temp download path: %s"), *BuildDownloadPath);
+	}
 }
 
 void FAndroidPlatformBackgroundHttpManager::Shutdown()
@@ -1280,5 +1300,4 @@ void JNICALL UE::Jni::Download::FUEDownloadWorker::nativeAndroidBackgroundDownlo
 	FAndroidBackgroundDownloadDelegates::AndroidBackgroundDownload_OnTickWorkerThread.ExecuteIfBound(jenv, thiz);
 }
 
-
 #undef LOCTEXT_NAMESPACE
\ No newline at end of file
diff --git a/Engine/Source/Runtime/Online/BackgroundHTTPFileHash/BackgroundHTTPFileHash.Build.cs b/Engine/Source/Runtime/Online/BackgroundHTTPFileHash/BackgroundHTTPFileHash.Build.cs
--- a/Engine/Source/Runtime/Online/BackgroundHTTPFileHash/BackgroundHTTPFileHash.Build.cs
+++ b/Engine/Source/Runtime/Online/BackgroundHTTPFileHash/BackgroundHTTPFileHash.Build.cs
@@ -12,7 +12,13 @@ public class BackgroundHTTPFileHash : ModuleRules
                 "Core",
             }
         );
+		PrivateDependencyModuleNames.AddRange(
+			new string[]
+			{
+				"BuildSettings",
+			}
+		);
 
-        CppCompileWarningSettings.UnsafeTypeCastWarningLevel = WarningLevel.Error;
+		CppCompileWarningSettings.UnsafeTypeCastWarningLevel = WarningLevel.Error;
     }
 }
diff --git a/Engine/Source/Runtime/Online/BackgroundHTTPFileHash/Private/BackgroundHttpFileHashHelper.cpp b/Engine/Source/Runtime/Online/BackgroundHTTPFileHash/Private/BackgroundHttpFileHashHelper.cpp
--- a/Engine/Source/Runtime/Online/BackgroundHTTPFileHash/Private/BackgroundHttpFileHashHelper.cpp
+++ b/Engine/Source/Runtime/Online/BackgroundHTTPFileHash/Private/BackgroundHttpFileHashHelper.cpp
@@ -11,6 +11,8 @@
 #include "Misc/FileHelper.h"
 #include "Misc/Paths.h"
 
+#include "BuildSettings.h"
+
 DEFINE_LOG_CATEGORY(LogBackgroundHttpFileHash);
 
 FString FURLTempFileMapping::ToString() const
@@ -94,7 +96,10 @@ const FString& FBackgroundHttpFileHashHelper::GetURLMappingFilePathNoLock()
 FString FBackgroundHttpFileHashHelper::GetFullPathOfTempFilename(const FString& TempFilename)
 {
 	// don't take lock here, as other methods taking lock call it
-	const FString FullFileName = FPaths::SetExtension(FPaths::GetCleanFilename(TempFilename), GetTempFileExtension());
+	FString FullFileName = FPaths::SetExtension(FPaths::GetCleanFilename(TempFilename), GetTempFileExtension());
+#if PLATFORM_ANDROID
+	FullFileName = FPaths::Combine(BuildSettings::GetBranchName(), FullFileName);
+#endif
 	return FPaths::Combine(GetTemporaryRootPath(), FullFileName);
 }
 
