Build fix

Original breakage introduced in 47953379

#rnx

diff --git a/Engine/Plugins/Online/Android/AndroidFetchBackgroundDownload/Source/AndroidFetchBackgroundDownload/Public/AndroidPlatformBackgroundHttpManager.h b/Engine/Plugins/Online/Android/AndroidFetchBackgroundDownload/Source/AndroidFetchBackgroundDownload/Public/AndroidPlatformBackgroundHttpManager.h
--- a/Engine/Plugins/Online/Android/AndroidFetchBackgroundDownload/Source/AndroidFetchBackgroundDownload/Public/AndroidPlatformBackgroundHttpManager.h
+++ b/Engine/Plugins/Online/Android/AndroidFetchBackgroundDownload/Source/AndroidFetchBackgroundDownload/Public/AndroidPlatformBackgroundHttpManager.h
@@ -38,6 +38,7 @@ public:
 	virtual void Initialize() override;
 	virtual void Shutdown() override;
 	virtual void SetCellularPreference(int32 Value) override;
+	virtual void SetLimitedDataPreference(int32 Value) override;
 	virtual bool IsGenericImplementation() const override { return false; }
 	bool Tick(float DeltaTime);
 	
