Pause new Fetch implementation properly when internet connection disappears

#rnx
#Synced-CL 45098147

diff --git a/Engine/Plugins/Online/Android/AndroidFetchBackgroundDownload/Source/AndroidFetchBackgroundDownload/Java/download/UEDownloadWorker.java b/Engine/Plugins/Online/Android/AndroidFetchBackgroundDownload/Source/AndroidFetchBackgroundDownload/Java/download/UEDownloadWorker.java
--- a/Engine/Plugins/Online/Android/AndroidFetchBackgroundDownload/Source/AndroidFetchBackgroundDownload/Java/download/UEDownloadWorker.java
+++ b/Engine/Plugins/Online/Android/AndroidFetchBackgroundDownload/Source/AndroidFetchBackgroundDownload/Java/download/UEDownloadWorker.java
@@ -155,16 +155,21 @@ public class UEDownloadWorker extends UEWorker implements DownloadProgressListen
 
 		if (null == NetworkListener)
 		{
-
 			NetworkListener = new NetworkConnectivityClient.Listener() {
 				@Override
 				public void onNetworkAvailable(NetworkConnectivityClient.NetworkTransportType networkTransportType) {
 					bLostNetwork = false;
+					if (mFetchManager != null) {
+						mFetchManager.NetworkAvailable();
+					}
 				}
 
 				@Override
 				public void onNetworkLost() {
 					bLostNetwork = true;
+					if (mFetchManager != null) {
+						mFetchManager.NetworkLost();
+					}
 				}
 			};
 			NetworkChangedManager.getInstance().addListener(NetworkListener);
@@ -217,7 +222,7 @@ public class UEDownloadWorker extends UEWorker implements DownloadProgressListen
 		mFetchManager.EnqueueRequests(getApplicationContext(),QueueDescription);
 
 		//Enter actual loop until work is finished
-		Log.verbose("Entering OnWorkerStart Loop waiting for Fetch2");
+		Log.verbose("Entering OnWorkerStart Loop");
 		try 
 		{
 			while (bReceivedResult == false)
@@ -358,12 +363,16 @@ public class UEDownloadWorker extends UEWorker implements DownloadProgressListen
 			//Keeping the code path to insert a watch dog later on
 			//mFetchManager.RequestCheckDownloadsStillActive(this);
 
-			try
-			{
-				nativeAndroidBackgroundDownloadOnTick();
-			}
-			catch (UnsatisfiedLinkError e)
+			// avoid spamming the log with endless linker errors
+			if (bGameThreadIsActive)
 			{
+				try
+				{				
+					nativeAndroidBackgroundDownloadOnTick();
+				}
+				catch (UnsatisfiedLinkError e)
+				{
+				}
 			}
 		}
 	}
@@ -775,6 +784,12 @@ public class UEDownloadWorker extends UEWorker implements DownloadProgressListen
 	@Override
 	public void OnDownloadProgress(String RequestID, long BytesWrittenSinceLastCall, long TotalBytesWritten)
 	{
+		if (!bGameThreadIsActive)
+		{
+			// avoid spamming the log with endless linker errors
+			return;
+		}
+
 		try
 		{
 			nativeAndroidBackgroundDownloadOnProgress(RequestID, BytesWrittenSinceLastCall, TotalBytesWritten);
@@ -799,6 +814,12 @@ public class UEDownloadWorker extends UEWorker implements DownloadProgressListen
 	@Override
 	public void OnDownloadComplete(String RequestID, String CompleteLocation, EDownloadCompleteReason CompleteReason)
 	{
+		if (!bGameThreadIsActive)
+		{
+			// avoid spamming the log with endless linker errors
+			return;
+		}
+
 		boolean bWasSuccess = (CompleteReason == EDownloadCompleteReason.Success);
 		try
 		{
@@ -812,6 +833,12 @@ public class UEDownloadWorker extends UEWorker implements DownloadProgressListen
 	@Override
 	public void OnDownloadMetrics(String RequestID, long TotalBytesDownloaded, long DownloadDuration, long DownloadStartTimeUTC, long DownloadEndTimeUTC)
 	{
+		if (!bGameThreadIsActive)
+		{
+			// avoid spamming the log with endless linker errors
+			return;
+		}
+
 		try
 		{
 			nativeAndroidBackgroundDownloadOnMetrics(RequestID, TotalBytesDownloaded, DownloadDuration, DownloadStartTimeUTC, DownloadEndTimeUTC);
@@ -840,12 +867,16 @@ public class UEDownloadWorker extends UEWorker implements DownloadProgressListen
 				editor.commit();
 			}
 
-			try
-			{
-				nativeAndroidBackgroundDownloadOnAllComplete(bDidAllRequestsSucceed);
-			}
-			catch (UnsatisfiedLinkError e)
+			// avoid spamming the log with endless linker errors
+			if (bGameThreadIsActive)
 			{
+				try
+				{
+					nativeAndroidBackgroundDownloadOnAllComplete(bDidAllRequestsSucceed);
+				}
+				catch (UnsatisfiedLinkError e)
+				{
+				}
 			}
 		
 			//If UE code didn't provide a result for the work in the above callback(IE: Engine isn't running yet, we are completely in background, etc.) 
diff --git a/Engine/Plugins/Online/Android/AndroidFetchBackgroundDownload/Source/AndroidFetchBackgroundDownload/Java/download/asyncdownloader/AsyncDownloader.java b/Engine/Plugins/Online/Android/AndroidFetchBackgroundDownload/Source/AndroidFetchBackgroundDownload/Java/download/asyncdownloader/AsyncDownloader.java
--- a/Engine/Plugins/Online/Android/AndroidFetchBackgroundDownload/Source/AndroidFetchBackgroundDownload/Java/download/asyncdownloader/AsyncDownloader.java
+++ b/Engine/Plugins/Online/Android/AndroidFetchBackgroundDownload/Source/AndroidFetchBackgroundDownload/Java/download/asyncdownloader/AsyncDownloader.java
@@ -28,7 +28,7 @@ public class AsyncDownloader {
 	// ---------------------------------------------------------------------------------------------------------------------	
 	
 	public AsyncDownloader(boolean AllowHTTP2, int MaxConcurrent) {
-		Log.debug("AsyncDownloader: AllowHTTP2=" + AllowHTTP2 + " | MaxConcurrent=" + MaxConcurrent);
+		//Log.debug("AsyncDownloader: AllowHTTP2=" + AllowHTTP2 + " | MaxConcurrent=" + MaxConcurrent);
 		this.Dispatcher = new Dispatcher();
 		this.Dispatcher.setMaxRequests(Math.max(1, MaxConcurrent));
 		this.Dispatcher.setMaxRequestsPerHost(Math.max(1, MaxConcurrent));
@@ -41,22 +41,22 @@ public class AsyncDownloader {
 				.build();
 	}
 
-	public void SetMaxConcurrent(int MaxConcurrent) {
-		Log.debug("SetMaxConcurrent: MaxConcurrent=" + MaxConcurrent);
+	public synchronized void SetMaxConcurrent(int MaxConcurrent) {
+		//Log.debug("SetMaxConcurrent: MaxConcurrent=" + MaxConcurrent);
 		Dispatcher.setMaxRequests(Math.max(1, MaxConcurrent));
 		Dispatcher.setMaxRequestsPerHost(Math.max(1, MaxConcurrent));
 	}
 
-	public void SetAsyncDownloaderListener(AsyncDownloaderListener Listner) { 
+	public synchronized void SetAsyncDownloaderListener(AsyncDownloaderListener Listner) { 
 		this.DownloaderListener = Listner;
 	}
 
-	public boolean Submit(AsyncDownloadRequest Request, AsyncDownloadRequestListener Listener) {
-		// Log.debug("Submit: Id=" + Request.Id);
+	public synchronized boolean Submit(AsyncDownloadRequest Request, AsyncDownloadRequestListener Listener) {
+		//Log.debug("Submit: Id=" + Request.Id);
 		DownloadEntry Existing = Entries.get(Request.Id);
 		if (Existing != null) {
 			Existing.Listener = Listener;
-			Log.debug("Submit: pre-Existing, updated Listener");
+			//Log.debug("Submit: pre-Existing, updated Listener");
 			return false;
 		}
 		DownloadEntry Entry = new DownloadEntry(Request, Listener);
@@ -65,12 +65,16 @@ public class AsyncDownloader {
 		return true;
 	}
 
-	public void Pause(String Id) {
-		// Log.debug("Pause: Id=" + Id);
+	public synchronized void Pause(String Id) {
+		//Log.debug("Pause: Id=" + Id);
 		DownloadEntry Entry = Entries.get(Id);
 		if (Entry == null) {
 			return;
 		}
+		if (Entry.State == EState.PAUSED) {
+			//Log.debug("Pause: Id=" + Id + " is already paused, ignoring call");
+			return;
+		}
 		Entry.UserPaused = true;
 		Call c = Entry.currentCall;
 		if (c != null) {
@@ -78,55 +82,47 @@ public class AsyncDownloader {
 		}
 		Entry.State = EState.PAUSED;
 		Entry.Listener.OnPaused(Id);
-		CheckIfAllDownloadsCompleted();
+		Entry.Attempts = 0;
 	}
 
-	public void ResumeAll() {
-		// Log.debug("ResumeAll");
-		PausedAll.set(false);
+	public synchronized void PauseAll() {
+		if (PausedAll.get()) {
+			return;
+		}
+		//Log.debug("PauseAll");
+		PausedAll.set(true);
 		for (DownloadEntry Entry : Entries.values()) {
-			if (Entry.State == EState.PAUSED || Entry.State == EState.PENDING || Entry.State == EState.FAILED) {
-				Entry.UserPaused = false;
-				StartOrQueue(Entry);
-			}
+			Pause(Entry.Request.Id);
 		}
 	}
 
-
-	public void Resume(String Id) {
-		// Log.debug("Resume: Id=" + Id);
+	public synchronized void Resume(String Id) {
+		//Log.debug("Resume: Id=" + Id);
 		DownloadEntry Entry = Entries.get(Id);
 		if (Entry == null) {
 			return;
 		}
-		Entry.UserPaused = false;
-		if (Entry.State == EState.COMPLETED) {
+		if (Entry.State != EState.PAUSED) {
+			//Log.debug("Resume: Id=" + Id + " not paused, ignoring call");
 			return;
 		}
+		Entry.UserPaused = false;
 		StartOrQueue(Entry);
 	}
 
-	public void PauseAll() {
-		// Log.debug("PauseAll");
-
-		PausedAll.set(true);
-
+	public synchronized void ResumeAll() {
+		if (!PausedAll.get()) {
+			return;
+		}
+		//Log.debug("ResumeAll");
+		PausedAll.set(false);
 		for (DownloadEntry Entry : Entries.values()) {
-			Entry.UserPaused = true;
-			Call c = Entry.currentCall;
-			if (c != null) {
-				c.cancel();
-			}
-			if (Entry.State == EState.RUNNING || Entry.State == EState.PENDING) {
-				Entry.State = EState.PAUSED;
-				Entry.Listener.OnPaused(Entry.Request.Id);
-			}
+			Resume(Entry.Request.Id);
 		}
-		CheckIfAllDownloadsCompleted();
 	}
 
-	public void Cancel(String Id) {
-		// Log.debug("Cancel: Id=" + Id);
+	public synchronized void Cancel(String Id) {
+		//Log.debug("Cancel: Id=" + Id);
 		DownloadEntry Entry = Entries.remove(Id);
 		if (Entry == null) {
 			return;
@@ -139,8 +135,8 @@ public class AsyncDownloader {
 		CheckIfAllDownloadsCompleted();
 	}
 
-	public void CancelAll() {
-		// Log.debug("CancelAll");
+	public synchronized void CancelAll() {
+		//Log.debug("CancelAll");
 
 		// do not fire the AllDownloadsComplete if we're cancelling them all
 		bAllowCheckIfAllDownloadsCompleted = false;
@@ -153,8 +149,8 @@ public class AsyncDownloader {
 		}
 	}
 
-	public void Retry(String Id) {
-		// Log.debug("Retry: Id=" + Id);
+	public synchronized void Retry(String Id) {
+		//Log.debug("Retry: Id=" + Id);
 
 		DownloadEntry OldEntry = Entries.get(Id);
 		if (OldEntry == null) {
@@ -179,7 +175,7 @@ public class AsyncDownloader {
 		StartOrQueue(NewEntry);
 	}
 
-	public float GetGroupProgress(int GroupID) {
+	public synchronized float GetGroupProgress(int GroupID) {
 		long TotalBytes = 0;
 		long DownloadedBytes = 0;
 
@@ -217,24 +213,24 @@ public class AsyncDownloader {
 	}
 
 	private void StartOrQueue(DownloadEntry Entry) {
-		// Log.debug("StartOrQueue: Id=" + Entry.Request.Id + " | TotalSize=" + Entry.Request.TotalSize);
+		//Log.debug("StartOrQueue: Id=" + Entry.Request.Id + " | TotalSize=" + Entry.Request.TotalSize);
 		if (Entry.State == EState.COMPLETED) {
-			// Log.debug("StartOrQueue: State set to completed for id " + Entry.Request.Id);
+			//Log.debug("StartOrQueue: State set to completed for id " + Entry.Request.Id);
 			return;
 		}
 		if (Entry.State == EState.CANCELLED) {
-			// Log.debug("StartOrQueue: State set to cancelled for id " + Entry.Request.Id);
+			//Log.debug("StartOrQueue: State set to cancelled for id " + Entry.Request.Id);
 			return;
 		}
 		if (Entry.UserPaused) {
-			// Log.debug("StartOrQueue: State set to user-paused for id " + Entry.Request.Id);
+			//Log.debug("StartOrQueue: State set to user-paused for id " + Entry.Request.Id);
 			return;
 		}
 		if (PausedAll.get()) {
-			// Log.debug("StartOrQueue: PausedAll set to true, not starting on Id " + Entry.Request.Id);
+			//Log.debug("StartOrQueue: PausedAll set to true, not starting on Id " + Entry.Request.Id);
 			return;
 		}
-		// Log.debug("StartOrQueue: started work on Id " + Entry.Request.Id);
+		//Log.debug("StartOrQueue: started work on Id " + Entry.Request.Id);
 
 		try {
 			final File file = Entry.Request.TargetFile;
@@ -252,7 +248,7 @@ public class AsyncDownloader {
 			if (CurrentLen > 0) {
 				if (CurrentLen <= Entry.Request.TotalSize) {
 					rb.header("Range", "bytes=" + CurrentLen + "-");
-					Log.debug("StartOrQueue: requesting partial content for id " + Entry.Request.Id + " | Offset " + CurrentLen);
+					//Log.debug("StartOrQueue: requesting partial content for id " + Entry.Request.Id + " | Offset " + CurrentLen);
 				} else {
 					// be on the safe side: delete the file completely if it is bigger than we expected
 					Log.error("StartOrQueue: requesting full file because current file size is too big (" + CurrentLen + " > " + Entry.Request.TotalSize + ")");
@@ -275,7 +271,7 @@ public class AsyncDownloader {
 				}
 
 				@Override public void onResponse(Call Call, Response Response) {
-					// Log.debug("onResponse: begin downloading | Id=" + Entry.Request.Id);
+					//Log.debug("onResponse: begin downloading | Id=" + Entry.Request.Id);
 					try (Response response = Response) {
 						int Code = response.code();
 						if (!response.isSuccessful()) {
@@ -351,7 +347,7 @@ public class AsyncDownloader {
 								Entry.Listener.OnProgress(Entry.Request.Id, WritePos, TotalExpected);
 							}
 
-							// Log.debug("onResponse: download completed | Id=" + Entry.Request.Id);
+							//Log.debug("onResponse: download completed | Id=" + Entry.Request.Id);
 							Entry.State = EState.COMPLETED;
 							Entry.currentCall = null;
 							Entry.Listener.OnSuccess(Entry.Request.Id, file);
@@ -370,10 +366,10 @@ public class AsyncDownloader {
 	}
 
 	private void HandleAttemptFailure(DownloadEntry Entry, Throwable Error, @Nullable URL LastUrl) {
-		// Log.debug("HandleAttemptFailure: Id=" + Entry.Request.Id);
+		//Log.debug("HandleAttemptFailure: Id=" + Entry.Request.Id);
 		Entry.currentCall = null;
 		if (Entry.UserPaused || PausedAll.get()) {
-			// Log.debug("HandleAttemptFailure: Id=" + Entry.Request.Id + " | User parked or PausedAll true, parking");
+			//Log.debug("HandleAttemptFailure: Id=" + Entry.Request.Id + " | User parked or PausedAll true, parking");
 			Entry.State = EState.PAUSED;
 			Entry.Listener.OnPaused(Entry.Request.Id);
 			CheckIfAllDownloadsCompleted();
@@ -385,10 +381,10 @@ public class AsyncDownloader {
 
 		final boolean bHaveRetriesRemaining = Attempt <= Entry.Request.MaxRetries;
 		if (bHaveRetriesRemaining) {
-			// Log.debug("HandleAttemptFailure: Id=" + Entry.Request.Id + " | trying next URL");
+			//Log.debug("HandleAttemptFailure: Id=" + Entry.Request.Id + " | trying next URL");
 			schedule(() -> StartOrQueue(Entry), 1000L);
 		} else {
-			// Log.debug("HandleAttemptFailure: Id=" + Entry.Request.Id + " | out of retries: failing download");
+			//Log.debug("HandleAttemptFailure: Id=" + Entry.Request.Id + " | out of retries: failing download");
 
 			Entry.State = EState.FAILED;
 			Entry.Listener.OnError(Entry.Request.Id, Error, LastUrl);
@@ -397,7 +393,7 @@ public class AsyncDownloader {
 		}
 	}
 
-	private void CheckIfAllDownloadsCompleted() {
+	private synchronized void CheckIfAllDownloadsCompleted() {
 		Log.debug("CheckIfAllDownloadsCompleted: Entries.size()=" + Entries.size());
 		if (!bAllowCheckIfAllDownloadsCompleted) {
 			return;
@@ -405,13 +401,13 @@ public class AsyncDownloader {
 		if (!Entries.isEmpty()) {
 			return;
 		}
-		Log.debug("CheckIfAllDownloadsCompleted: done, firing onAllCompleted");
+		//Log.debug("CheckIfAllDownloadsCompleted: done, firing onAllCompleted");
 		AsyncDownloaderListener Listener = DownloaderListener;
 		if (Listener != null) {
 			Log.debug("CheckIfAllDownloadsCompleted: firing listener");
 			Listener.OnAllCompleted();
 		} else {
-			Log.debug("CheckIfAllDownloadsCompleted: no listener set!");
+			//Log.debug("CheckIfAllDownloadsCompleted: no listener set!");
 		}
 	}
 
diff --git a/Engine/Plugins/Online/Android/AndroidFetchBackgroundDownload/Source/AndroidFetchBackgroundDownload/Java/download/fetch/FetchManagerInterface.java b/Engine/Plugins/Online/Android/AndroidFetchBackgroundDownload/Source/AndroidFetchBackgroundDownload/Java/download/fetch/FetchManagerInterface.java
--- a/Engine/Plugins/Online/Android/AndroidFetchBackgroundDownload/Source/AndroidFetchBackgroundDownload/Java/download/fetch/FetchManagerInterface.java
+++ b/Engine/Plugins/Online/Android/AndroidFetchBackgroundDownload/Source/AndroidFetchBackgroundDownload/Java/download/fetch/FetchManagerInterface.java
@@ -25,4 +25,6 @@ public interface FetchManagerInterface
 	public void RetryDownload(String RequestID);
 	public void RequestGroupProgressUpdate(int GroupID, DownloadProgressListener ListenerToUpdate);
 	public void SetAllDownloadsCompletedListener(@NonNull AllDownloadsCompletedListener Listener);
+	default public void NetworkAvailable() {}
+	default public void NetworkLost() {}
 }
diff --git a/Engine/Plugins/Online/Android/AndroidFetchBackgroundDownload/Source/AndroidFetchBackgroundDownload/Java/download/fetch/FetchManagerNew.java b/Engine/Plugins/Online/Android/AndroidFetchBackgroundDownload/Source/AndroidFetchBackgroundDownload/Java/download/fetch/FetchManagerNew.java
--- a/Engine/Plugins/Online/Android/AndroidFetchBackgroundDownload/Source/AndroidFetchBackgroundDownload/Java/download/fetch/FetchManagerNew.java
+++ b/Engine/Plugins/Online/Android/AndroidFetchBackgroundDownload/Source/AndroidFetchBackgroundDownload/Java/download/fetch/FetchManagerNew.java
@@ -175,6 +175,23 @@ public class FetchManagerNew implements FetchManagerInterface
 		SetupAllDownloadsCompletedListener_Internal();
 	}
 
+	@Override
+	public void NetworkAvailable() 
+	{
+		if (AsyncDownloader == null) {
+			return;
+		}
+		AsyncDownloader.ResumeAll();
+	}
+
+	@Override
+	public void NetworkLost() {
+		if (AsyncDownloader == null) {
+			return;
+		}
+		AsyncDownloader.PauseAll();
+	}
+
 	// ---------------------------------------------------------------------------------------------------------------------
 	// private data + methods
 	// ---------------------------------------------------------------------------------------------------------------------
