Various fixes to BackgroundHTTP, mainly for android

#Synced-CL 48867346

diff --git a/Engine/Plugins/Online/Android/AndroidFetchBackgroundDownload/Source/AndroidFetchBackgroundDownload/Java/download/DownloadProgressListener.java b/Engine/Plugins/Online/Android/AndroidFetchBackgroundDownload/Source/AndroidFetchBackgroundDownload/Java/download/DownloadProgressListener.java
--- a/Engine/Plugins/Online/Android/AndroidFetchBackgroundDownload/Source/AndroidFetchBackgroundDownload/Java/download/DownloadProgressListener.java
+++ b/Engine/Plugins/Online/Android/AndroidFetchBackgroundDownload/Source/AndroidFetchBackgroundDownload/Java/download/DownloadProgressListener.java
@@ -6,10 +6,10 @@ import com.epicgames.unreal.download.UEDownloadWorker.EDownloadCompleteReason;
 //Interface for a class that wants to be notified when a download completes
 public interface DownloadProgressListener
 {
-	public void OnDownloadProgress(String RequestID, long BytesWrittenSinceLastCall, long TotalBytesWritten);
+	public void OnDownloadProgress(String DestinationLocation, long TotalBytesWritten, long BytesWrittenSinceLastCall);
 	public void OnDownloadGroupProgress(int GroupID, int Progress, boolean Indeterminate);
-	public void OnDownloadComplete(String RequestID, String CompleteLocation, EDownloadCompleteReason CompleteReason);
-	public void OnDownloadMetrics(String RequestID, long TotalBytesDownloaded, long DownloadDuration, long DownloadStartTimeUTC, long DownloadEndTimeUTC);
+	public void OnDownloadComplete(String DestinationLocation, EDownloadCompleteReason CompleteReason);
+	public void OnDownloadMetrics(String DestinationLocation, long TotalBytesDownloaded, long DownloadDuration, long DownloadStartTimeUTC, long DownloadEndTimeUTC);
 	public void OnAllDownloadsComplete(boolean bDidAllRequestsSucceed);
-	public void OnDownloadEnqueued(String RequestID, boolean bEnqueueSuccess);
+	public void OnDownloadEnqueued(String DestinationLocation, boolean bEnqueueSuccess);
 }
\ No newline at end of file
diff --git a/Engine/Plugins/Online/Android/AndroidFetchBackgroundDownload/Source/AndroidFetchBackgroundDownload/Java/download/UEDownloadWorker.java b/Engine/Plugins/Online/Android/AndroidFetchBackgroundDownload/Source/AndroidFetchBackgroundDownload/Java/download/UEDownloadWorker.java
--- a/Engine/Plugins/Online/Android/AndroidFetchBackgroundDownload/Source/AndroidFetchBackgroundDownload/Java/download/UEDownloadWorker.java
+++ b/Engine/Plugins/Online/Android/AndroidFetchBackgroundDownload/Source/AndroidFetchBackgroundDownload/Java/download/UEDownloadWorker.java
@@ -782,7 +782,7 @@ public class UEDownloadWorker extends UEWorker implements DownloadProgressListen
 	// DownloadCompletionListener Implementation
 	//
 	@Override
-	public void OnDownloadProgress(String RequestID, long BytesWrittenSinceLastCall, long TotalBytesWritten)
+	public void OnDownloadProgress(String DestinationLocation, long TotalBytesWritten, long BytesWrittenSinceLastCall)
 	{
 		if (!bGameThreadIsActive)
 		{
@@ -792,7 +792,7 @@ public class UEDownloadWorker extends UEWorker implements DownloadProgressListen
 
 		try
 		{
-			nativeAndroidBackgroundDownloadOnProgress(RequestID, BytesWrittenSinceLastCall, TotalBytesWritten);
+			nativeAndroidBackgroundDownloadOnProgress(DestinationLocation, TotalBytesWritten, BytesWrittenSinceLastCall);
 		}
 		catch (UnsatisfiedLinkError e)
 		{
@@ -812,7 +812,7 @@ public class UEDownloadWorker extends UEWorker implements DownloadProgressListen
 	}
 	
 	@Override
-	public void OnDownloadComplete(String RequestID, String CompleteLocation, EDownloadCompleteReason CompleteReason)
+	public void OnDownloadComplete(String DestinationLocation, EDownloadCompleteReason CompleteReason)
 	{
 		if (!bGameThreadIsActive)
 		{
@@ -823,7 +823,7 @@ public class UEDownloadWorker extends UEWorker implements DownloadProgressListen
 		boolean bWasSuccess = (CompleteReason == EDownloadCompleteReason.Success);
 		try
 		{
-			nativeAndroidBackgroundDownloadOnComplete(RequestID, CompleteLocation, bWasSuccess);
+			nativeAndroidBackgroundDownloadOnComplete(DestinationLocation, bWasSuccess);
 		}
 		catch (UnsatisfiedLinkError e)
 		{
@@ -831,7 +831,7 @@ public class UEDownloadWorker extends UEWorker implements DownloadProgressListen
 	}
 
 	@Override
-	public void OnDownloadMetrics(String RequestID, long TotalBytesDownloaded, long DownloadDuration, long DownloadStartTimeUTC, long DownloadEndTimeUTC)
+	public void OnDownloadMetrics(String DestinationLocation, long TotalBytesDownloaded, long DownloadDuration, long DownloadStartTimeUTC, long DownloadEndTimeUTC)
 	{
 		if (!bGameThreadIsActive)
 		{
@@ -841,7 +841,7 @@ public class UEDownloadWorker extends UEWorker implements DownloadProgressListen
 
 		try
 		{
-			nativeAndroidBackgroundDownloadOnMetrics(RequestID, TotalBytesDownloaded, DownloadDuration, DownloadStartTimeUTC, DownloadEndTimeUTC);
+			nativeAndroidBackgroundDownloadOnMetrics(DestinationLocation, TotalBytesDownloaded, DownloadDuration, DownloadStartTimeUTC, DownloadEndTimeUTC);
 		}
 		catch (UnsatisfiedLinkError e)
 		{
@@ -900,16 +900,16 @@ public class UEDownloadWorker extends UEWorker implements DownloadProgressListen
 	}
 
 	@Override
-	public void OnDownloadEnqueued(String RequestID, boolean bEnqueueSuccess)
+	public void OnDownloadEnqueued(String DestinationLocation, boolean bEnqueueSuccess)
 	{
 		if (bEnqueueSuccess)
 		{
-			Log.verbose("Enqueue success:%s" + RequestID);
+			Log.verbose("Enqueue success:" + DestinationLocation);
 		}
 		else
 		{
-			Log.debug("Enqueue failure, retrying request:" + RequestID);
-			mFetchManager.RetryDownload(RequestID);
+			Log.debug("Enqueue failure, retrying request:" + DestinationLocation);
+			mFetchManager.RetryDownload(DestinationLocation);
 		}
 
 		bHasEnqueueHappened = true;
@@ -943,9 +943,9 @@ public class UEDownloadWorker extends UEWorker implements DownloadProgressListen
 	//
 	// Functions called by our UE c++ code on this object
 	//
-	public void PauseRequest(String RequestID)
+	public void PauseRequest(String DestinationLocation)
 	{
-		Log.debug("@*@ PauseRequest: " + RequestID);
+		Log.debug("@*@ PauseRequest: " + DestinationLocation);
 
 		boolean ShouldHandleCellular = false;
 		if (NotificationDescription != null)
@@ -960,21 +960,21 @@ public class UEDownloadWorker extends UEWorker implements DownloadProgressListen
 				bWaitingForCellularApproval = true;
 			}
 		}
-		mFetchManager.PauseDownload(RequestID, true);
+		mFetchManager.PauseDownload(DestinationLocation, true);
 	}
 	
-	public void ResumeRequest(String RequestID)
+	public void ResumeRequest(String DestinationLocation)
 	{
-		Log.debug("@*@ ResumeRequest: " + RequestID);
+		Log.debug("@*@ ResumeRequest: " + DestinationLocation);
 
 		// When the C++ resumes the request all network issues should be solved
 		bWaitingForCellularApproval = false;
-		mFetchManager.PauseDownload(RequestID, false);
+		mFetchManager.PauseDownload(DestinationLocation, false);
 	}
 	
-	public void CancelRequest(String RequestID)
+	public void CancelRequest(String DestinationLocation)
 	{
-		mFetchManager.CancelDownload(RequestID);
+		mFetchManager.CancelDownload(DestinationLocation);
 	}
 
 	// This is a way of figuring out if the worker was spawned when the app was killed or not
@@ -987,7 +987,7 @@ public class UEDownloadWorker extends UEWorker implements DownloadProgressListen
 	native void nativeAndroidBackgroundDownloadOnWorkerStart(String WorkID);
 	native void nativeAndroidBackgroundDownloadOnWorkerStop(String WorkID);
 	native void nativeAndroidBackgroundDownloadOnProgress(String TaskID, long BytesWrittenSinceLastCall, long TotalBytesWritten);
-	native void nativeAndroidBackgroundDownloadOnComplete(String TaskID, String CompleteLocation, boolean bWasSuccess);
+	native void nativeAndroidBackgroundDownloadOnComplete(String TaskID, boolean bWasSuccess);
 	native void nativeAndroidBackgroundDownloadOnMetrics(String TaskID, long TotalBytesDownloaded, long DownloadDuration, long DownloadStartTimeUTC, long DownloadEndTimeUTC);
 	native void nativeAndroidBackgroundDownloadOnAllComplete(boolean bDidAllRequestsSucceed);
 	native void nativeAndroidBackgroundDownloadOnTick();
diff --git a/Engine/Plugins/Online/Android/AndroidFetchBackgroundDownload/Source/AndroidFetchBackgroundDownload/Java/download/asyncdownloader/AsyncDownloadRequest.java b/Engine/Plugins/Online/Android/AndroidFetchBackgroundDownload/Source/AndroidFetchBackgroundDownload/Java/download/asyncdownloader/AsyncDownloadRequest.java
--- a/Engine/Plugins/Online/Android/AndroidFetchBackgroundDownload/Source/AndroidFetchBackgroundDownload/Java/download/asyncdownloader/AsyncDownloadRequest.java
+++ b/Engine/Plugins/Online/Android/AndroidFetchBackgroundDownload/Source/AndroidFetchBackgroundDownload/Java/download/asyncdownloader/AsyncDownloadRequest.java
@@ -11,7 +11,8 @@ import java.util.Map;
 import java.util.Objects;
 
 final public class AsyncDownloadRequest {
-    public final String Id;
+    public final String Id;								// internal unique Id used for tracking
+    public final String RequestId;						// optional request Id used for deduplication
     public final List<String> Urls;						// cycle through these on failure
     public final int MaxRetries;						// total retry attempts across all URLs
     public final File TargetFile;
@@ -22,6 +23,7 @@ final public class AsyncDownloadRequest {
 
     private AsyncDownloadRequest(Builder b) {
         this.Id = Objects.requireNonNull(b.Id, "id");
+        this.RequestId = b.RequestId;
         this.Urls = List.copyOf(Objects.requireNonNull(b.Urls, "urls"));
         if (this.Urls.isEmpty()) 
 			throw new IllegalArgumentException("urls must not be empty");
@@ -35,6 +37,7 @@ final public class AsyncDownloadRequest {
 
     public static class Builder {
         private String Id;
+        private String RequestId;
         private List<String> Urls = new ArrayList<>();
         private int MaxRetries = 3;
         private File TargetFile;
@@ -47,6 +50,10 @@ final public class AsyncDownloadRequest {
 			this.Id = Id; 
 			return this; 
 		}
+        public Builder SetRequestId(String RequestId) { 
+			this.Id = RequestId; 
+			return this; 
+		}
         public Builder SetUrls(List<String> Urls) { 
 			this.Urls = new ArrayList<>(Urls); 
 			return this; 
diff --git a/Engine/Plugins/Online/Android/AndroidFetchBackgroundDownload/Source/AndroidFetchBackgroundDownload/Java/download/asyncdownloader/AsyncDownloader.java b/Engine/Plugins/Online/Android/AndroidFetchBackgroundDownload/Source/AndroidFetchBackgroundDownload/Java/download/asyncdownloader/AsyncDownloader.java
--- a/Engine/Plugins/Online/Android/AndroidFetchBackgroundDownload/Source/AndroidFetchBackgroundDownload/Java/download/asyncdownloader/AsyncDownloader.java
+++ b/Engine/Plugins/Online/Android/AndroidFetchBackgroundDownload/Source/AndroidFetchBackgroundDownload/Java/download/asyncdownloader/AsyncDownloader.java
@@ -54,15 +54,24 @@ public class AsyncDownloader {
 	public synchronized boolean Submit(AsyncDownloadRequest Request, AsyncDownloadRequestListener Listener) {
 		//Log.debug("Submit: Id=" + Request.Id);
 		DownloadEntry Existing = Entries.get(Request.Id);
-		if (Existing != null) {
-			Existing.Listener = Listener;
-			//Log.debug("Submit: pre-Existing, updated Listener");
-			return false;
+		if (Existing == null) {
+			if (Request.RequestId == null || Request.RequestId.isEmpty()) {
+				DownloadEntry Entry = new DownloadEntry(Request, Listener);
+				Entries.put(Request.Id, Entry);
+				StartOrQueue(Entry);
+				return true;
+			}
+			Existing = Entries.values().stream().filter(Entry -> Request.RequestId.equals(Entry.Request.RequestId)).findFirst().orElse(null);
+			if (Existing == null) {
+				DownloadEntry Entry = new DownloadEntry(Request, Listener);
+				Entries.put(Request.Id, Entry);
+				StartOrQueue(Entry);
+				return true;
+			}
 		}
-		DownloadEntry Entry = new DownloadEntry(Request, Listener);
-		Entries.put(Request.Id, Entry);
-		StartOrQueue(Entry);
-		return true;
+		Existing.Listener = Listener;
+		//Log.debug("Submit: pre-Existing, updated Listener");
+		return false;
 	}
 
 	public synchronized void Pause(String Id) {
@@ -245,7 +254,7 @@ public class AsyncDownloader {
 				}
 			}
 
-			if (CurrentLen > 0) {
+			if (CurrentLen > 0 && Entry.Request.TotalSize > 0) {
 				if (CurrentLen <= Entry.Request.TotalSize) {
 					rb.header("Range", "bytes=" + CurrentLen + "-");
 					//Log.debug("StartOrQueue: requesting partial content for id " + Entry.Request.Id + " | Offset " + CurrentLen);
@@ -310,7 +319,7 @@ public class AsyncDownloader {
 						} else {
 							TotalContentLen = ContentLen;
 						}
-						if (TotalContentLen != Entry.Request.TotalSize) {
+						if (TotalContentLen != Entry.Request.TotalSize && Entry.Request.TotalSize > 0) {
 							Log.error("onResponse: error: content length mismatch: TotalContentLen="+TotalContentLen+" != TotalExpectedSize="+Entry.Request.TotalSize+" | ContentLen="+ContentLen+" | FileOnDiskLen=" + FileOnDiskLen);
 							HandleAttemptFailure(Entry,
 									new IOException("HTTP " + Code + " " + response.message()),
diff --git a/Engine/Plugins/Online/Android/AndroidFetchBackgroundDownload/Source/AndroidFetchBackgroundDownload/Java/download/datastructs/DownloadDescription.java b/Engine/Plugins/Online/Android/AndroidFetchBackgroundDownload/Source/AndroidFetchBackgroundDownload/Java/download/datastructs/DownloadDescription.java
--- a/Engine/Plugins/Online/Android/AndroidFetchBackgroundDownload/Source/AndroidFetchBackgroundDownload/Java/download/datastructs/DownloadDescription.java
+++ b/Engine/Plugins/Online/Android/AndroidFetchBackgroundDownload/Source/AndroidFetchBackgroundDownload/Java/download/datastructs/DownloadDescription.java
@@ -241,7 +241,7 @@ public class DownloadDescription
 
 			if (false == jObject.isNull(TotalBytesNeededKey))
 			{
-				ReturnedDescription.TotalBytesNeeded = jObject.optInt(TotalBytesNeededKey);
+				ReturnedDescription.TotalBytesNeeded = jObject.optLong(TotalBytesNeededKey);
 			}
 		}
 		catch( JSONException e )
diff --git a/Engine/Plugins/Online/Android/AndroidFetchBackgroundDownload/Source/AndroidFetchBackgroundDownload/Java/download/fetch/FetchEnqueueResultListener.java b/Engine/Plugins/Online/Android/AndroidFetchBackgroundDownload/Source/AndroidFetchBackgroundDownload/Java/download/fetch/FetchEnqueueResultListener.java
--- a/Engine/Plugins/Online/Android/AndroidFetchBackgroundDownload/Source/AndroidFetchBackgroundDownload/Java/download/fetch/FetchEnqueueResultListener.java
+++ b/Engine/Plugins/Online/Android/AndroidFetchBackgroundDownload/Source/AndroidFetchBackgroundDownload/Java/download/fetch/FetchEnqueueResultListener.java
@@ -11,18 +11,18 @@ import com.tonyodev.fetch2core.Func;
 public interface FetchEnqueueResultListener
 {
 	//Function called on this FetchEnqueueResultListener by the FetchEnqueueRequestCallback
-	public void OnFetchEnqueueRequestCallback(@NonNull String RequestID, @NonNull Request EnqueuedRequest);
+	public void OnFetchEnqueueRequestCallback(@NonNull String DestinationLocation, @NonNull Request EnqueuedRequest);
 	
 	//Function called on this FetchEnqueueResultListener by the FetchEnqueueErrorCallback
-	public void OnFetchEnqueueErrorCallback(@NonNull String RequestID, @NonNull Error EnqueueError);
+	public void OnFetchEnqueueErrorCallback(@NonNull String DestinationLocation, @NonNull Error EnqueueError);
 	
 	//Helper class passed in to the Fetch.enqueue command to pass back the results of the enqueue operation on success.
 	public class FetchEnqueueRequestCallback implements Func<Request>
 	{
-		public FetchEnqueueRequestCallback(@NonNull FetchEnqueueResultListener Owner, @NonNull String CachedRequestID)
+		public FetchEnqueueRequestCallback(@NonNull FetchEnqueueResultListener Owner, @NonNull String DestinationLocation)
 		{
 			this.Owner = Owner;
-			this.CachedRequestID = CachedRequestID;
+			this.DestinationLocation = DestinationLocation;
 		}
 
 		@Override
@@ -30,22 +30,22 @@ public interface FetchEnqueueResultListener
 		{
 			if (Owner != null) 
 			{
-				Owner.OnFetchEnqueueRequestCallback(CachedRequestID, result);
+				Owner.OnFetchEnqueueRequestCallback(DestinationLocation, result);
 			}			
 			Owner = null;
 		}
 
 		FetchEnqueueResultListener Owner = null;
-		String CachedRequestID = null;
+		String DestinationLocation = null;
 	}
 
 	//Helper class passed in to the Fetch.enqueue command to pass back the results of the enqueue operation on error
 	public class FetchEnqueueErrorCallback implements Func<Error>
 	{
-		public FetchEnqueueErrorCallback(@NonNull FetchEnqueueResultListener Owner, @NonNull String CachedRequestID)
+		public FetchEnqueueErrorCallback(@NonNull FetchEnqueueResultListener Owner, @NonNull String DestinationLocation)
 		{
 			this.Owner = Owner;
-			this.CachedRequestID = CachedRequestID;
+			this.DestinationLocation = DestinationLocation;
 		}
 
 		@Override
@@ -53,12 +53,12 @@ public interface FetchEnqueueResultListener
 		{
 			if (Owner != null) 
 			{
-				Owner.OnFetchEnqueueErrorCallback(CachedRequestID, error);
+				Owner.OnFetchEnqueueErrorCallback(DestinationLocation, error);
 			}
 			Owner = null;
 		}
 		
 		FetchEnqueueResultListener Owner = null;
-		String CachedRequestID = null;
+		String DestinationLocation = null;
 	}
 }
\ No newline at end of file
diff --git a/Engine/Plugins/Online/Android/AndroidFetchBackgroundDownload/Source/AndroidFetchBackgroundDownload/Java/download/fetch/FetchManagerInterface.java b/Engine/Plugins/Online/Android/AndroidFetchBackgroundDownload/Source/AndroidFetchBackgroundDownload/Java/download/fetch/FetchManagerInterface.java
--- a/Engine/Plugins/Online/Android/AndroidFetchBackgroundDownload/Source/AndroidFetchBackgroundDownload/Java/download/fetch/FetchManagerInterface.java
+++ b/Engine/Plugins/Online/Android/AndroidFetchBackgroundDownload/Source/AndroidFetchBackgroundDownload/Java/download/fetch/FetchManagerInterface.java
@@ -18,11 +18,11 @@ public interface FetchManagerInterface
 
 	public void StopWork(String WorkID);
 	public void EnqueueRequests(@NonNull Context context, @NonNull DownloadQueueDescription QueueDescription);
-	public void PauseDownload(String RequestID, boolean bPause);
+	public void PauseDownload(String DestinationLocation, boolean bPause);
 	public void PauseAllDownloads();
 	public void ResumeAllDownloads();
-	public void CancelDownload(String RequestID);
-	public void RetryDownload(String RequestID);
+	public void CancelDownload(String DestinationLocation);
+	public void RetryDownload(String DestinationLocation);
 	public void RequestGroupProgressUpdate(int GroupID, DownloadProgressListener ListenerToUpdate);
 	public void SetAllDownloadsCompletedListener(@NonNull AllDownloadsCompletedListener Listener);
 	default public void NetworkAvailable() {}
diff --git a/Engine/Plugins/Online/Android/AndroidFetchBackgroundDownload/Source/AndroidFetchBackgroundDownload/Java/download/fetch/FetchManagerNew.java b/Engine/Plugins/Online/Android/AndroidFetchBackgroundDownload/Source/AndroidFetchBackgroundDownload/Java/download/fetch/FetchManagerNew.java
--- a/Engine/Plugins/Online/Android/AndroidFetchBackgroundDownload/Source/AndroidFetchBackgroundDownload/Java/download/fetch/FetchManagerNew.java
+++ b/Engine/Plugins/Online/Android/AndroidFetchBackgroundDownload/Source/AndroidFetchBackgroundDownload/Java/download/fetch/FetchManagerNew.java
@@ -82,7 +82,8 @@ public class FetchManagerNew implements FetchManagerInterface
 
 			final boolean bSuccess = AsyncDownloader.Submit(
 				new AsyncDownloadRequest.Builder()
-					.SetId(DownloadDescription.RequestID)
+					.SetId(DownloadDescription.DestinationLocation)
+					.SetRequestId(DownloadDescription.RequestID)
 					.SetUrls(List.copyOf(DownloadDescription.URLs))
 					.SetMaxRetries(DownloadDescription.MaxRetryCount)
 					.SetGroupID(DownloadDescription.GroupID)
@@ -91,7 +92,7 @@ public class FetchManagerNew implements FetchManagerInterface
 					.SetTotalSize(DownloadDescription.TotalBytesNeeded)
 					.Build(),
 				new AsyncDownloadRequestListenerImpl(this, GlobalDownloadProgressListener, TargetLocation, DownloadLocation));
-			GlobalDownloadProgressListener.OnDownloadEnqueued(DownloadDescription.RequestID, bSuccess);
+			GlobalDownloadProgressListener.OnDownloadEnqueued(DownloadDescription.DestinationLocation, bSuccess);
 		}
 
 		// clear flag that stops downloader from downloading
@@ -99,15 +100,15 @@ public class FetchManagerNew implements FetchManagerInterface
 	}
 
 	@Override
-	public synchronized void PauseDownload(String RequestID, boolean bPause)
+	public synchronized void PauseDownload(String DestinationLocation, boolean bPause)
 	{
 		if (AsyncDownloader == null) {
 			return;
 		}
 		if (bPause) {
-			AsyncDownloader.Pause(RequestID);
+			AsyncDownloader.Pause(DestinationLocation);
 		} else {
-			AsyncDownloader.Resume(RequestID);
+			AsyncDownloader.Resume(DestinationLocation);
 		}
 	}
 
@@ -130,21 +131,21 @@ public class FetchManagerNew implements FetchManagerInterface
 	}
 
 	@Override
-	public synchronized void CancelDownload(String RequestID)
+	public synchronized void CancelDownload(String DestinationLocation)
 	{
 		if (AsyncDownloader == null) {
 			return;
 		}
-		AsyncDownloader.Cancel(RequestID);
+		AsyncDownloader.Cancel(DestinationLocation);
 	}
 
 	@Override
-	public synchronized void RetryDownload(String RequestID)
+	public synchronized void RetryDownload(String DestinationLocation)
 	{
 		if (AsyncDownloader == null) {
 			return;
 		}
-		AsyncDownloader.Retry(RequestID);
+		AsyncDownloader.Retry(DestinationLocation);
 	}
 
 	@Override
@@ -293,7 +294,7 @@ public class FetchManagerNew implements FetchManagerInterface
 			}
 
 			ProgressListener.OnDownloadMetrics(Id, TotalBytes, Duration, StartTime, EndTime);
-			ProgressListener.OnDownloadComplete(Id, TargetLocation, Reason);
+			ProgressListener.OnDownloadComplete(Id, Reason);
 		}
 
 		private long LastDownloadedSoFar = 0;
diff --git a/Engine/Plugins/Online/Android/AndroidFetchBackgroundDownload/Source/AndroidFetchBackgroundDownload/Java/download/fetch/FetchManagerOld.java b/Engine/Plugins/Online/Android/AndroidFetchBackgroundDownload/Source/AndroidFetchBackgroundDownload/Java/download/fetch/FetchManagerOld.java
--- a/Engine/Plugins/Online/Android/AndroidFetchBackgroundDownload/Source/AndroidFetchBackgroundDownload/Java/download/fetch/FetchManagerOld.java
+++ b/Engine/Plugins/Online/Android/AndroidFetchBackgroundDownload/Source/AndroidFetchBackgroundDownload/Java/download/fetch/FetchManagerOld.java
@@ -115,7 +115,7 @@ public class FetchManagerOld implements FetchDownloadProgressOwner, FetchEnqueue
 	public synchronized void EnqueueRequests(Context context, DownloadQueueDescription QueueDescription)
 	{
 		runningState = ETaskRunningState.Unknown;
-		InitFetch(context);
+		InitFetch(context, QueueDescription.MaxConcurrentDownloads);
 		
 		int NumDownloads = QueueDescription.DownloadDescriptions.size();
 		Log.debug("EnqueueRequests called with " + NumDownloads + " DownloadDescriptions. Current RequestedDownloads Size: " + RequestedDownloads.size());
@@ -142,8 +142,8 @@ public class FetchManagerOld implements FetchDownloadProgressOwner, FetchEnqueue
 	{
 		//Create reference copy of RequestedDownloads and clear original lists for
 		//RequestedDownloads, CompletedDownloads, and FailedDownloads to allow them to re-populate with new work state
-		HashMap<String,DownloadDescription> TempRequestedDownloadsCopy =  new HashMap<String,DownloadDescription>(RequestedDownloads);
-		RequestedDownloads.clear();
+		HashMap<String, DownloadDescription> TempRequestedDownloadsCopy = RequestedDownloads;
+		RequestedDownloads = new HashMap<String, DownloadDescription>();
 		CompletedDownloads.clear();
 		FailedDownloads.clear();
 
@@ -154,34 +154,38 @@ public class FetchManagerOld implements FetchDownloadProgressOwner, FetchEnqueue
 			
 			//Set any DownloadDescription settings that need to come from our QueueDescription
 			NewDownloadDescription.ProgressListener = QueueDescription.ProgressListener;
-							
-			String RequestIDKey = NewDownloadDescription.RequestID;
-			//Don't even have an entry for this RequestID in the old list, so it's completely new
+			
+			DownloadDescription ActiveDescription = TempRequestedDownloadsCopy.get(NewDownloadDescription.DestinationLocation);
 			boolean LastRequest = (NumNewDownloadDescriptions-1) == DescriptionIndex;
-			if (!TempRequestedDownloadsCopy.containsKey(RequestIDKey))
+			
+			if (ActiveDescription == null)
 			{
-				QueueNewDownloadDescription(NewDownloadDescription, LastRequest);
+				if (NewDownloadDescription.RequestID == null || NewDownloadDescription.RequestID.isEmpty())
+				{
+					//Don't even have an entry in the old list, so it's completely new
+					QueueNewDownloadDescription(NewDownloadDescription, LastRequest);
+					continue;
+				}
+				
+				ActiveDescription = TempRequestedDownloadsCopy.values().stream().filter(Description -> NewDownloadDescription.RequestID.equals(Description.RequestID)).findFirst().orElse(null);
+				
+				if (ActiveDescription == null)
+				{
+					//Don't even have an entry for this RequestID in the old list, so it's completely new
+					QueueNewDownloadDescription(NewDownloadDescription, LastRequest);
+					continue;
+				}
 			}
+			
 			//Need to update our DownloadDescription
-			else
-			{
-				DownloadDescription ActiveDescription = TempRequestedDownloadsCopy.get(RequestIDKey);
-				HandleChangedDownloadDescription(ActiveDescription, NewDownloadDescription, LastRequest);
-			}
+			HandleChangedDownloadDescription(ActiveDescription, NewDownloadDescription, LastRequest);
 		}
 	}
 	
 	public synchronized void ResetDownloadCompleteTracking(DownloadDescription Description)
 	{
-		if (CompletedDownloads.containsKey(Description.RequestID))
-		{
-			CompletedDownloads.remove(Description.RequestID);
-		}
-
-		if (FailedDownloads.containsKey(Description.RequestID))
-		{
-			FailedDownloads.remove(Description.RequestID);
-		}
+		CompletedDownloads.remove(Description.DestinationLocation);
+		FailedDownloads.remove(Description.DestinationLocation);
 	}
 
 	public synchronized ECompleteReason GetCompleteReasonForDownload(DownloadDescription Description)
@@ -191,7 +195,7 @@ public class FetchManagerOld implements FetchDownloadProgressOwner, FetchEnqueue
 			return ECompleteReason.Cancelled;
 		}
 
-		if (FailedDownloads.containsKey(Description.RequestID))
+		if (FailedDownloads.containsKey(Description.DestinationLocation))
 		{
 			return ECompleteReason.Error;
 		}
@@ -208,13 +212,13 @@ public class FetchManagerOld implements FetchDownloadProgressOwner, FetchEnqueue
 			Log.error("Call to QueueNewDownloadDescription after FetchInstance is invalidated! RequestID:" + Description.RequestID);
 			return;
 		}
-		RequestedDownloads.put(Description.RequestID, Description);
+		RequestedDownloads.put(Description.DestinationLocation, Description);
 		ResetDownloadCompleteTracking(Description);
 
 		if (!CheckForPreviouslyCompletedDownload(Description))
 		{	
-			FetchEnqueueResultListener.FetchEnqueueRequestCallback RequestCallback = new FetchEnqueueResultListener.FetchEnqueueRequestCallback(this, Description.RequestID);
-			FetchEnqueueResultListener.FetchEnqueueErrorCallback ErrorCallback = new FetchEnqueueResultListener.FetchEnqueueErrorCallback(this, Description.RequestID);
+			FetchEnqueueResultListener.FetchEnqueueRequestCallback RequestCallback = new FetchEnqueueResultListener.FetchEnqueueRequestCallback(this, Description.DestinationLocation);
+			FetchEnqueueResultListener.FetchEnqueueErrorCallback ErrorCallback = new FetchEnqueueResultListener.FetchEnqueueErrorCallback(this, Description.DestinationLocation);
 		
 			Request FetchRequest = BuildFetchRequest(Description);
 			if (FetchRequest == null)
@@ -247,7 +251,7 @@ public class FetchManagerOld implements FetchDownloadProgressOwner, FetchEnqueue
 
 	private synchronized boolean CheckForPreviouslyCompletedDownload(DownloadDescription Description)
 	{
-		boolean bInCompleteList = CompletedDownloads.containsKey(Description.RequestID);
+		boolean bInCompleteList = CompletedDownloads.containsKey(Description.DestinationLocation);
 		
 		File CompleteLocation = new File(Description.DestinationLocation);
 		boolean bHasFileInCompleteLocation = CompleteLocation.exists();
@@ -256,21 +260,21 @@ public class FetchManagerOld implements FetchDownloadProgressOwner, FetchEnqueue
 	}
 
 	@Override
-	public synchronized void PauseDownload(String RequestID, boolean bPause)
+	public synchronized void PauseDownload(String DestinationLocation, boolean bPause)
 	{
 		runningState = ETaskRunningState.Unknown;
 
 		//Early out if we have invalidated our FetchInstance and thus would crash if trying to do work
 		if (!IsFetchInstanceValid())
 		{
-			Log.error("Call to PauseDownload after FetchInstance has been invalidated. Skipping as no work can be done. RequestID: " + RequestID);
+			Log.error("Call to PauseDownload after FetchInstance has been invalidated. Skipping as no work can be done. DestinationLocation: " + DestinationLocation);
 			return;
 		}	
 
-		DownloadDescription MatchedDesc = RequestedDownloads.get(RequestID);
+		DownloadDescription MatchedDesc = RequestedDownloads.get(DestinationLocation);
 		if (null == MatchedDesc)
 		{
-			Log.error("No DownloadDescription found for RequestID " + RequestID +" during PauseDownload!");
+			Log.error("No DownloadDescription found for DestinationLocation " + DestinationLocation +" during PauseDownload!");
 			return;
 		}
 		
@@ -303,7 +307,6 @@ public class FetchManagerOld implements FetchDownloadProgressOwner, FetchEnqueue
 		int pauseCount = 0;
 		int totalCount = 0;
 
-		HashMap<String,DownloadDescription> TempRequestedDownloadsCopy =  new HashMap<String,DownloadDescription>(RequestedDownloads);
 		for(Map.Entry<String, DownloadDescription> entry : RequestedDownloads.entrySet()) 
 		{
 			
@@ -346,7 +349,6 @@ public class FetchManagerOld implements FetchDownloadProgressOwner, FetchEnqueue
 		int resumeCount = 0;
 		int totalCount = 0;
 		
-		HashMap<String,DownloadDescription> TempRequestedDownloadsCopy =  new HashMap<String,DownloadDescription>(RequestedDownloads);
 		for(Map.Entry<String, DownloadDescription> entry : RequestedDownloads.entrySet()) 
 		{
 			
@@ -372,28 +374,28 @@ public class FetchManagerOld implements FetchDownloadProgressOwner, FetchEnqueue
 	}
 
 	@Override
-	public synchronized void CancelDownload(String RequestID)
+	public synchronized void CancelDownload(String DestinationLocation)
 	{
 		if (!IsFetchInstanceValid())
 		{
-			Log.error("Call to CancelDownload with an invalid fetch instance. Skipping fetch work to avoid crashing. RequestID: " + RequestID);
+			Log.error("Call to CancelDownload with an invalid fetch instance. Skipping fetch work to avoid crashing. DestinationLocation: " + DestinationLocation);
 			return;
 		}
 
-		DownloadDescription MatchedDesc = RequestedDownloads.get(RequestID);
+		DownloadDescription MatchedDesc = RequestedDownloads.get(DestinationLocation);
 		if (null == MatchedDesc)
 		{
-			Log.error("No DownloadDescription found for RequestID " + RequestID +" during CancelDownload!");
+			Log.error("No DownloadDescription found for DestinationLocation " + DestinationLocation +" during CancelDownload!");
 			return;
 		}
 		
 		if (false == MatchedDesc.bIsCancelled)
 		{
-			Log.verbose("Cancelling download:" + RequestID);
+			Log.verbose("Cancelling download:" + MatchedDesc.RequestID);
 			
 			//Set download as complete since it was cancelled intentionally
 			MatchedDesc.PreviousDownloadPercent = 100;
-			CompletedDownloads.put(RequestID, MatchedDesc);
+			CompletedDownloads.put(DestinationLocation, MatchedDesc);
 			MatchedDesc.bIsCancelled = true;
 
 			FetchInstance.delete(MatchedDesc.CachedFetchID);
@@ -407,9 +409,9 @@ public class FetchManagerOld implements FetchDownloadProgressOwner, FetchEnqueue
 		String TempDownloadLoc = GetTempDownloadDestination(Description);
 
 		//Can not build a valid FetchRequest with a null URL (will crash when trying to enqueue)
-		if ((URL == null) || (TempDownloadLoc == null))
+		if ((URL == null) || (Description.DestinationLocation == null))
 		{
-			Log.error("URL or TempDownloadLoction to build Uri from are null! Can not create a valid FetchRequest! RequestID:" + Description.RequestID);
+			Log.error("URL or DestinationLocation to build Uri from are null! Can not create a valid FetchRequest! RequestID:" + Description.RequestID);
 			return null;
 		}
 
@@ -418,7 +420,7 @@ public class FetchManagerOld implements FetchDownloadProgressOwner, FetchEnqueue
 
 		Request FetchRequest = new Request(URL, DownloadUri);
 		FetchRequest.setPriority(GetFetchPriority(Description));
-		FetchRequest.setTag(Description.RequestID);
+		FetchRequest.setTag(Description.DestinationLocation);
 		FetchRequest.setGroupId(Description.GroupID);
 		FetchRequest.setAutoRetryMaxAttempts(Description.IndividualURLRetryCount);
 		
@@ -442,7 +444,7 @@ public class FetchManagerOld implements FetchDownloadProgressOwner, FetchEnqueue
 
 		//Handle the change by cancelling and recreating the fetch2 download
 		RecreateDownloadByTagFunc RecreateFunc = new RecreateDownloadByTagFunc(this, NewDescription, LastRequest);
-		FetchInstance.getDownloadsByTag(NewDescription.RequestID, RecreateFunc);
+		FetchInstance.getDownloadsByTag(OldDescription.DestinationLocation, RecreateFunc);
 	}
 	
 	//Copies over the non-serialized values that track download state from our old descrition to our new description so that changing
@@ -503,7 +505,7 @@ public class FetchManagerOld implements FetchDownloadProgressOwner, FetchEnqueue
 		}
 	}
 	
-	private synchronized void InitFetch(Context context)
+	private synchronized void InitFetch(Context context, int MaxConcurrentDownloads)
 	{
 		Log.debug("InitFetch called");
 
@@ -563,6 +565,7 @@ public class FetchManagerOld implements FetchDownloadProgressOwner, FetchEnqueue
 						.setLogger(FetchLog)
 						.enableRetryOnNetworkGain(true)
 						.setProgressReportingInterval(200)
+						.setDownloadConcurrentLimit(MaxConcurrentDownloads)
 						.build());
 				}
 				else
@@ -573,7 +576,7 @@ public class FetchManagerOld implements FetchDownloadProgressOwner, FetchEnqueue
 						.setLogger(FetchLog)
 						.enableRetryOnNetworkGain(true)
 						.setProgressReportingInterval(200)
-						.setDownloadConcurrentLimit(10)
+						.setDownloadConcurrentLimit(MaxConcurrentDownloads)
 						.setHttpDownloader(new OkHttpDownloader(okHttpClient))
 						.build());
 				}
@@ -730,7 +733,7 @@ public class FetchManagerOld implements FetchDownloadProgressOwner, FetchEnqueue
 	@Override
 	public void OnDownloadQueued(@NonNull Download download)
 	{
-		Log.verbose("OnDownloadQueued: " + GetRequestID(download));
+		Log.verbose("OnDownloadQueued: " + GetDestinationLocation(download));
 
 		//Treat this as just an un-pause for now to make sure we handle cases where Fetch resumes downloads that we paused (This can happen when retrying for networkgain as an example)
 		OnDownloadChangePauseState(download, false);
@@ -739,13 +742,13 @@ public class FetchManagerOld implements FetchDownloadProgressOwner, FetchEnqueue
 	@Override
 	public void OnDownloadStarted(@NonNull Download download)
 	{
-		String DownloadRequestID = GetRequestID(download);
-		Log.verbose("OnDownloadStarted: " + DownloadRequestID);
+		String DestinationLocation = GetDestinationLocation(download);
+		Log.verbose("OnDownloadStarted: " + DestinationLocation);
 
-		DownloadDescription MatchedDownload = RequestedDownloads.get(DownloadRequestID);
+		DownloadDescription MatchedDownload = RequestedDownloads.get(DestinationLocation);
 		if (null == MatchedDownload)
 		{
-			Log.error("OnDownloadStarted called from DownloadProgressOwner implementation with a download that doesn't match any DownloadDesc! Download's Tag: " + DownloadRequestID);
+			Log.error("OnDownloadStarted called from DownloadProgressOwner implementation with a download that doesn't match any DownloadDesc! Download's Tag: " + DestinationLocation);
 			return;
 		}
 
@@ -758,14 +761,14 @@ public class FetchManagerOld implements FetchDownloadProgressOwner, FetchEnqueue
 	@Override
 	public void OnDownloadProgress(@NonNull Download download, boolean indeterminate, long downloadedBytesPerSecond, long etaInMilliSeconds)
 	{
-		String DownloadRequestID = GetRequestID(download);
+		String DestinationLocation = GetDestinationLocation(download);
 
-		Log.verbose("OnDownloadProgress: " + DownloadRequestID + " bytes/s:" + downloadedBytesPerSecond + " eta:" + etaInMilliSeconds);
+		Log.verbose("OnDownloadProgress: " + DestinationLocation + " bytes/s:" + downloadedBytesPerSecond + " eta:" + etaInMilliSeconds);
 
-		DownloadDescription MatchedDownload = RequestedDownloads.get(DownloadRequestID);
+		DownloadDescription MatchedDownload = RequestedDownloads.get(DestinationLocation);
 		if (null == MatchedDownload)
 		{
-			Log.error("OnDownloadProgress called from DownloadProgressOwner implementation with a download that doesn't match any DownloadDesc! Download's Tag: " + DownloadRequestID);
+			Log.error("OnDownloadProgress called from DownloadProgressOwner implementation with a download that doesn't match any DownloadDesc! Download's Tag: " + DestinationLocation);
 			return;
 		}
 		
@@ -776,18 +779,18 @@ public class FetchManagerOld implements FetchDownloadProgressOwner, FetchEnqueue
 			long TotalDownloaded = download.getDownloaded();
 			//If our TotalDownload size has gone down, we likely had an error and restarted the download since the previous update.
 			//Don't want to return a negative value, so just use 0 for now
-			if (MatchedDownload.PreviousDownloadedBytes > TotalDownloaded)
+			if (MatchedDownload.TotalDownloadedBytes > TotalDownloaded)
 			{
 				TotalDownloadedSinceLastCall = 0;
 			}
 			else
 			{
-				TotalDownloadedSinceLastCall = (TotalDownloaded - MatchedDownload.PreviousDownloadedBytes);
+				TotalDownloadedSinceLastCall = (TotalDownloaded - MatchedDownload.TotalDownloadedBytes);
 			}
 			MatchedDownload.PreviousDownloadPercent = download.getProgress();
 			MatchedDownload.PreviousDownloadedBytes = TotalDownloadedSinceLastCall;
 			MatchedDownload.TotalDownloadedBytes = TotalDownloaded;
-			MatchedDownload.ProgressListener.OnDownloadProgress(GetRequestID(download), TotalDownloadedSinceLastCall, TotalDownloaded);
+			MatchedDownload.ProgressListener.OnDownloadProgress(GetDestinationLocation(download), TotalDownloaded, TotalDownloadedSinceLastCall);
 			MatchedDownload.TotalBytesNeeded = download.getTotal();
 			
 		}
@@ -802,26 +805,26 @@ public class FetchManagerOld implements FetchDownloadProgressOwner, FetchEnqueue
 	{
 		runningState = ETaskRunningState.Unknown;
 
-		String DownloadRequestID = GetRequestID(download);
-		Log.verbose("OnDownloadChangePauseState: " + DownloadRequestID);
+		String DestinationLocation = GetDestinationLocation(download);
+		Log.verbose("OnDownloadChangePauseState: " + DestinationLocation);
 
-		DownloadDescription MatchedDownload = RequestedDownloads.get(DownloadRequestID);
+		DownloadDescription MatchedDownload = RequestedDownloads.get(DestinationLocation);
 		if (null == MatchedDownload)
 		{
-			Log.error("OnDownloadChangePauseState called from DownloadProgressOwner implementation with a download that doesn't match any DownloadDesc! Download's Tag: " + DownloadRequestID);
+			Log.error("OnDownloadChangePauseState called from DownloadProgressOwner implementation with a download that doesn't match any DownloadDesc! Download's Tag: " + DestinationLocation);
 			return;
 		}
 		
 		if (!IsFetchInstanceValid())
 		{
-			Log.error("Call to OnDownloadChangePauseState with an invalid fetch instance. Skipping fetch work to avoid crashing. RequestID: " + DownloadRequestID);
+			Log.error("Call to OnDownloadChangePauseState with an invalid fetch instance. Skipping fetch work to avoid crashing. RequestID: " + MatchedDownload.RequestID);
 			return;
 		}
 
 		//the DownloadDescription is always definitive, so make sure the new download state matches our expectations
 		if (bIsPaused != MatchedDownload.bIsPaused)
 		{
-			Log.verbose("OnDownloadChangePauseState required Fetch action. RequestID:" + DownloadRequestID + " New bIsPaused:" + MatchedDownload.bIsPaused + " Old Fetch bIsPaused:" + bIsPaused);			
+			Log.verbose("OnDownloadChangePauseState required Fetch action. RequestID:" + MatchedDownload.RequestID + " New bIsPaused:" + MatchedDownload.bIsPaused + " Old Fetch bIsPaused:" + bIsPaused);			
 			
 			if (MatchedDownload.bIsPaused)
 			{
@@ -864,13 +867,13 @@ public class FetchManagerOld implements FetchDownloadProgressOwner, FetchEnqueue
 	@Override
 	public void OnDownloadComplete(@NonNull Download download, ECompleteReason completeReason)
 	{
-		String DownloadRequestID = GetRequestID(download);
-		Log.debug("OnDownloadComplete: " + DownloadRequestID + " Reason:" + completeReason);
+		String DestinationLocation = GetDestinationLocation(download);
+		Log.debug("OnDownloadComplete: " + DestinationLocation + " Reason:" + completeReason);
 
-		DownloadDescription MatchedDownload = RequestedDownloads.get(DownloadRequestID);
+		DownloadDescription MatchedDownload = RequestedDownloads.get(DestinationLocation);
 		if (null == MatchedDownload)
 		{
-			Log.error("OnDownloadComplete called from DownloadProgressOwner implementation with a download that doesn't match any DownloadDesc! Download's Tag: " + DownloadRequestID);
+			Log.error("OnDownloadComplete called from DownloadProgressOwner implementation with a download that doesn't match any DownloadDesc! Download's Tag: " + DestinationLocation);
 			return;
 		}
 		
@@ -881,24 +884,24 @@ public class FetchManagerOld implements FetchDownloadProgressOwner, FetchEnqueue
 	// FetchEnqueueResultListener Implementation
 	//
 	@Override
-	public void OnFetchEnqueueRequestCallback(@NonNull String RequestID, @NonNull Request EnqueuedRequest)
+	public void OnFetchEnqueueRequestCallback(@NonNull String DestinationLocation, @NonNull Request EnqueuedRequest)
 	{
-		Log.verbose("Enqueued Request Success. ID:" + RequestID);
+		Log.verbose("Enqueued Request Success. DestinationLocation:" + DestinationLocation);
 		
-		DownloadDescription MatchingDescription = RequestedDownloads.get(RequestID);
+		DownloadDescription MatchingDescription = RequestedDownloads.get(DestinationLocation);
 		if (null == MatchingDescription)
 		{
-			Log.error("OnFetchEnqueueRequestCallback called on request is missing from our RequestedDownloads list! RequestID:" + RequestID);
+			Log.error("OnFetchEnqueueRequestCallback called on request is missing from our RequestedDownloads list! DestinationLocation:" + DestinationLocation);
 			return;
 		}
 
-		MatchingDescription.ProgressListener.OnDownloadEnqueued(RequestID, true);
+		MatchingDescription.ProgressListener.OnDownloadEnqueued(DestinationLocation, true);
 	}
 
 	@Override
-	public synchronized void OnFetchEnqueueErrorCallback(@NonNull String RequestID, @NonNull Error EnqueueError)
+	public synchronized void OnFetchEnqueueErrorCallback(@NonNull String DestinationLocation, @NonNull Error EnqueueError)
 	{
-		Log.error("Error Enqueing Request! " + RequestID + " Error: " + EnqueueError);
+		Log.error("Error Enqueing Request! " + DestinationLocation + " Error: " + EnqueueError);
 
 		Throwable thrown = EnqueueError.getThrowable();
 		if (thrown != null) {
@@ -911,25 +914,25 @@ public class FetchManagerOld implements FetchDownloadProgressOwner, FetchEnqueue
 		//If no real error supplied, lets check for fetch having already enqueued this in a separate listener
 		if (IsFetchInstanceValid() && IsErrorPossiblyIndicatingPreviouslyQueued(EnqueueError))
 		{
-			Log.error("Enqueue Error " + EnqueueError + " could be because of a previous enqueue already succeeding. Checking if Request is already enqueued before erroring. " + RequestID);
+			Log.error("Enqueue Error " + EnqueueError + " could be because of a previous enqueue already succeeding. Checking if Request is already enqueued before erroring. " + DestinationLocation);
 						
-			CheckForQueuedDownload CheckForQueueFunc = new CheckForQueuedDownload(this, RequestID);
-			FetchInstance.getDownloadsByTag(RequestID,CheckForQueueFunc);
+			CheckForQueuedDownload CheckForQueueFunc = new CheckForQueuedDownload(this, DestinationLocation);
+			FetchInstance.getDownloadsByTag(DestinationLocation, CheckForQueueFunc);
 		
 			return;
 		}
 
-		DownloadDescription MatchingDescription = RequestedDownloads.get(RequestID);
+		DownloadDescription MatchingDescription = RequestedDownloads.get(DestinationLocation);
 		if (null == MatchingDescription)
 		{
-			Log.error("OnFetchEnqueueErrorCallback called on request is missing from our RequestedDownloads list! RequestID:" + RequestID);
+			Log.error("OnFetchEnqueueErrorCallback called on request is missing from our RequestedDownloads list! DestinationLocation:" + DestinationLocation);
 			return;
 		}
 
-		MatchingDescription.ProgressListener.OnDownloadEnqueued(RequestID, false);
+		MatchingDescription.ProgressListener.OnDownloadEnqueued(DestinationLocation, false);
 	}
 	
-	public synchronized void HandleEnqueueCheckCallback(String RequestID, @Nullable Download FetchDownload)
+	public synchronized void HandleEnqueueCheckCallback(String DestinationLocation, @Nullable Download FetchDownload)
 	{
 		boolean bIsAlreadyEnqueued = false;
 
@@ -952,7 +955,7 @@ public class FetchManagerOld implements FetchDownloadProgressOwner, FetchEnqueue
 		}
 
 		
-		DownloadDescription MatchingDescription = RequestedDownloads.get(RequestID);
+		DownloadDescription MatchingDescription = RequestedDownloads.get(DestinationLocation);
 		
 		//If we weren't already enqueued or this request does't map to an existing download we should remove it from the fetch instance
 		//to try and allow any future retries to complete successfully
@@ -966,7 +969,7 @@ public class FetchManagerOld implements FetchDownloadProgressOwner, FetchEnqueue
 				
 		if (null == MatchingDescription)
 		{
-			Log.debug("No DownloadDescription matches " + RequestID + " can not send enqueue callback with result " + bIsAlreadyEnqueued);
+			Log.debug("No DownloadDescription matches " + DestinationLocation + " can not send enqueue callback with result " + bIsAlreadyEnqueued);
 			return;
 		}
 		
@@ -976,7 +979,7 @@ public class FetchManagerOld implements FetchDownloadProgressOwner, FetchEnqueue
 			MatchingDescription.CachedFetchID = FetchDownload.getId();
 		}
 
-		MatchingDescription.ProgressListener.OnDownloadEnqueued(RequestID, bIsAlreadyEnqueued);
+		MatchingDescription.ProgressListener.OnDownloadEnqueued(DestinationLocation, bIsAlreadyEnqueued);
 	}
 
 	public boolean IsErrorPossiblyIndicatingPreviouslyQueued(@NonNull Error EnqueueError)
@@ -1003,25 +1006,25 @@ public class FetchManagerOld implements FetchDownloadProgressOwner, FetchEnqueue
 		return false;
 	}
 
-	//if UEDownloadableWorker is calling into it, the only thing it should need is the RequestID
+	//if UEDownloadableWorker is calling into it, the only thing it should need is the DestinationLocation
 	@Override
-	public void RetryDownload(String RequestID)
+	public void RetryDownload(String DestinationLocation)
 	{
-		RetryDownload(RequestID, null);
+		RetryDownload(DestinationLocation, null);
 	}
 	
-	private synchronized void RetryDownload(String RequestID, @Nullable Download FetchDownload)
+	private synchronized void RetryDownload(String DestinationLocation, @Nullable Download FetchDownload)
 	{
-		DownloadDescription MatchingDescription = RequestedDownloads.get(RequestID);
+		DownloadDescription MatchingDescription = RequestedDownloads.get(DestinationLocation);
 		if (null == MatchingDescription)
 		{
-			Log.error("RetryDownload called on invalid download that was never requested! RequestID:" + RequestID);
+			Log.error("RetryDownload called on invalid download that was never requested! DestinationLocation:" + DestinationLocation);
 			return;
 		}
 
 		if (!IsFetchInstanceValid())
 		{
-			Log.debug("Call to RetryDownload with an invalid fetch instance. Skipping fetch work and completing to avoid crashing. RequestID:" + RequestID);
+			Log.debug("Call to RetryDownload with an invalid fetch instance. Skipping fetch work and completing to avoid crashing. RequestID:" + MatchingDescription.RequestID);
 			CompleteDownload(MatchingDescription, FetchRequestProgressListener.ECompleteReason.Error, true);
 			return;
 		}
@@ -1073,10 +1076,10 @@ public class FetchManagerOld implements FetchDownloadProgressOwner, FetchEnqueue
 
 	private class CheckForQueuedDownload implements Func<List<Download>>
 	{
-		public CheckForQueuedDownload (FetchManagerOld Owner, String CachedRequestID)
+		public CheckForQueuedDownload(FetchManagerOld Owner, String DestinationLocation)
 		{
 			this.Owner = Owner;
-			this.CachedRequestID = CachedRequestID;
+			this.DestinationLocation = DestinationLocation;
 		}
 
 		@Override
@@ -1085,7 +1088,7 @@ public class FetchManagerOld implements FetchDownloadProgressOwner, FetchEnqueue
 			//Callback for first found download
 			if ((MatchingDownloads != null) && (MatchingDownloads.size() > 0))
 			{
-				Owner.HandleEnqueueCheckCallback(CachedRequestID, MatchingDownloads.get(0));
+				Owner.HandleEnqueueCheckCallback(DestinationLocation, MatchingDownloads.get(0));
 			}
 			//remove and delete any other downloads that might have been found after the first
 			synchronized (this.Owner)
@@ -1105,7 +1108,7 @@ public class FetchManagerOld implements FetchDownloadProgressOwner, FetchEnqueue
 		}
 
 		private FetchManagerOld Owner;
-		private String CachedRequestID = null;
+		private String DestinationLocation = null;
 	}
 
 	private class RetryDownloadFunc implements Func2<Download>
@@ -1226,7 +1229,7 @@ public class FetchManagerOld implements FetchDownloadProgressOwner, FetchEnqueue
 		{
 			DownloadDesc.PreviousDownloadPercent = 100;
 		}
-		CompletedDownloads.put(DownloadDesc.RequestID, DownloadDesc);
+		CompletedDownloads.put(DownloadDesc.DestinationLocation, DownloadDesc);
 
 		//Only bubble up to the UEDownloadWorker non-intenional completes as we initiated any other completes internally
 		if (IsCompleteReasonIntentional(CompleteReason))
@@ -1268,13 +1271,13 @@ public class FetchManagerOld implements FetchDownloadProgressOwner, FetchEnqueue
 				
 					//Overwrite our successful complete reason with an error.
 					CompleteReasonToSend = EDownloadCompleteReason.Error;
-					FailedDownloads.put(DownloadDesc.RequestID, DownloadDesc);
+					FailedDownloads.put(DownloadDesc.DestinationLocation, DownloadDesc);
 				}
 			}
 		}
 		else
 		{
-			FailedDownloads.put(DownloadDesc.RequestID, DownloadDesc);
+			FailedDownloads.put(DownloadDesc.DestinationLocation, DownloadDesc);
 		}
 
 		long DownloadDuration = -1;
@@ -1287,9 +1290,9 @@ public class FetchManagerOld implements FetchDownloadProgressOwner, FetchEnqueue
 			DownloadDuration = ElapsedDownloadTime - DownloadDesc.TotalPausedTime;
 		}
 
-		DownloadDesc.ProgressListener.OnDownloadMetrics(DownloadDesc.RequestID, DownloadDesc.TotalDownloadedBytes, DownloadDuration, DownloadDesc.DownloadStartTime, DownloadFinishTime);
+		DownloadDesc.ProgressListener.OnDownloadMetrics(DownloadDesc.DestinationLocation, DownloadDesc.TotalDownloadedBytes, DownloadDuration, DownloadDesc.DownloadStartTime, DownloadFinishTime);
 
-		DownloadDesc.ProgressListener.OnDownloadComplete(DownloadDesc.RequestID, DownloadDesc.DestinationLocation, CompleteReasonToSend);
+		DownloadDesc.ProgressListener.OnDownloadComplete(DownloadDesc.DestinationLocation, CompleteReasonToSend);
 		
 		if (CheckForAllCompleted)
 		{
@@ -1338,7 +1341,7 @@ public class FetchManagerOld implements FetchDownloadProgressOwner, FetchEnqueue
 		}
 	}
 	
-	private String GetRequestID(Download download)
+	private String GetDestinationLocation(Download download)
 	{
 		if (null != download)
 		{
diff --git a/Engine/Plugins/Online/Android/AndroidFetchBackgroundDownload/Source/AndroidFetchBackgroundDownload/Private/AndroidJniDownloadUEDownloadWorker.h b/Engine/Plugins/Online/Android/AndroidFetchBackgroundDownload/Source/AndroidFetchBackgroundDownload/Private/AndroidJniDownloadUEDownloadWorker.h
--- a/Engine/Plugins/Online/Android/AndroidFetchBackgroundDownload/Source/AndroidFetchBackgroundDownload/Private/AndroidJniDownloadUEDownloadWorker.h
+++ b/Engine/Plugins/Online/Android/AndroidFetchBackgroundDownload/Source/AndroidFetchBackgroundDownload/Private/AndroidJniDownloadUEDownloadWorker.h
@@ -14,8 +14,8 @@ namespace UE::Jni::Download
 
 		static void JNICALL nativeAndroidBackgroundDownloadOnWorkerStart(JNIEnv* env, jobject thiz, jstring WorkID);
 		static void JNICALL nativeAndroidBackgroundDownloadOnWorkerStop(JNIEnv* env, jobject thiz, jstring WorkID);
-		static void JNICALL nativeAndroidBackgroundDownloadOnProgress(JNIEnv* env, jobject thiz, jstring TaskID, jlong BytesWrittenSinceLastCall, jlong TotalBytesWritten);
-		static void JNICALL nativeAndroidBackgroundDownloadOnComplete(JNIEnv* env, jobject thiz, jstring TaskID, jstring CompleteLocation, jboolean bWasSuccess);
+		static void JNICALL nativeAndroidBackgroundDownloadOnProgress(JNIEnv* env, jobject thiz, jstring TaskID, jlong TotalBytesWritten, jlong BytesWrittenSinceLastCall);
+		static void JNICALL nativeAndroidBackgroundDownloadOnComplete(JNIEnv* env, jobject thiz, jstring TaskID, jboolean bWasSuccess);
 		static void JNICALL nativeAndroidBackgroundDownloadOnMetrics(JNIEnv* env, jobject thiz, jstring TaskID, jlong TotalBytesDownloaded, jlong DownloadDuration, jlong DownloadStartTimeUTC, jlong DownloadEndTimeUTC);
 		static void JNICALL nativeAndroidBackgroundDownloadOnAllComplete(JNIEnv* env, jobject thiz, jboolean bDidAllRequestsSucceed);
 		static void JNICALL nativeAndroidBackgroundDownloadOnTick(JNIEnv* env, jobject thiz);
diff --git a/Engine/Plugins/Online/Android/AndroidFetchBackgroundDownload/Source/AndroidFetchBackgroundDownload/Private/AndroidPlatformBackgroundHttpManager.cpp b/Engine/Plugins/Online/Android/AndroidFetchBackgroundDownload/Source/AndroidFetchBackgroundDownload/Private/AndroidPlatformBackgroundHttpManager.cpp
--- a/Engine/Plugins/Online/Android/AndroidFetchBackgroundDownload/Source/AndroidFetchBackgroundDownload/Private/AndroidPlatformBackgroundHttpManager.cpp
+++ b/Engine/Plugins/Online/Android/AndroidFetchBackgroundDownload/Source/AndroidFetchBackgroundDownload/Private/AndroidPlatformBackgroundHttpManager.cpp
@@ -61,9 +61,9 @@ namespace FAndroidBackgroundDownloadDelegates
 {
 	DECLARE_DELEGATE_TwoParams(FAndroidBackgroundDownload_OnWorkerStart, FString /*WorkID*/, jobject /*UnderlyingWorker*/);
 	DECLARE_DELEGATE_TwoParams(FAndroidBackgroundDownload_OnWorkerStop, FString /*WorkID*/, jobject /*UnderlyingWorker*/);
-	DECLARE_DELEGATE_FourParams(FAndroidBackgroundDownload_OnProgress, jobject /*UnderlyingWorker*/, FString /*RequestID*/, int64_t /*BytesWrittenSinceLastCall*/, int64_t /*TotalBytesWritten*/);
-	DECLARE_DELEGATE_FourParams(FAndroidBackgroundDownload_OnComplete, jobject /*UnderlyingWorker*/, FString /*RequestID*/, FString /*CompleteLocation*/, bool /*bWasSuccess*/);
-	DECLARE_DELEGATE_SixParams(FAndroidBackgroundDownload_OnMetrics, jobject /*UnderlyingWorker*/, FString /*RequestID*/, int64_t /*TotalBytesDownloaded*/, int64_t /*DownloadDuration*/, int64_t /*DownloadStartTimeUTC*/, int64_t /*DownloadEndTimeUTC*/);
+	DECLARE_DELEGATE_FourParams(FAndroidBackgroundDownload_OnProgress, jobject /*UnderlyingWorker*/, FString /*DestinationLocation*/, int64_t /*TotalBytesWritten*/, int64_t /*BytesWrittenSinceLastCall*/);
+	DECLARE_DELEGATE_ThreeParams(FAndroidBackgroundDownload_OnComplete, jobject /*UnderlyingWorker*/, FString /*DestinationLocation*/, bool /*bWasSuccess*/);
+	DECLARE_DELEGATE_SixParams(FAndroidBackgroundDownload_OnMetrics, jobject /*UnderlyingWorker*/, FString /*DestinationLocation*/, int64_t /*TotalBytesDownloaded*/, int64_t /*DownloadDuration*/, int64_t /*DownloadStartTimeUTC*/, int64_t /*DownloadEndTimeUTC*/);
 	DECLARE_DELEGATE_TwoParams(FAndroidBackgroundDownload_OnAllComplete, jobject /*UnderlyingWorker*/, bool /*bDidAllRequestsSucceed*/);
 	DECLARE_DELEGATE_TwoParams(FAndroidBackgroundDownload_OnTickWorkerThread, JNIEnv*, jobject /*UnderlyingWorker*/);
 
@@ -323,7 +323,7 @@ void FAndroidPlatformBackgroundHttpManager::ActivatePendingRequests()
 					
 					UE_LOG(LogBackgroundHttpManager, Display, TEXT("Adding %d currently ActiveRequests back to RequestsToSendToJavaLayer"), ActiveRequests.Num());
 
-					for (FBackgroundHttpRequestPtr& Request : ActiveRequests)
+					for (const FBackgroundHttpRequestPtr& Request : ActiveRequests)
 					{
 						FAndroidBackgroundHttpRequestPtr AndroidRequest = StaticCastSharedPtr<FAndroidPlatformBackgroundHttpRequest>(Request);
 						if (ensureAlwaysMsgf(AndroidRequest.IsValid(), TEXT("Unexpected illegal non-Android request in ActiveRequests!")))
@@ -473,33 +473,33 @@ void FAndroidPlatformBackgroundHttpManager::HandlePendingCompletes()
 	const bool bShouldProcessCompletes = FPlatformAtomics::InterlockedExchange(&bHasPendingCompletes, false);
 	if (bShouldProcessCompletes)
 	{
-		TArray<FBackgroundHttpRequestPtr> RequestsToComplete;
+		TArray<FAndroidBackgroundHttpRequestPtr> RequestsToComplete;
 		//Populate list of requests waiting on completion handling
 		{
 			FRWScopeLock ScopeLock(ActiveRequestLock, SLT_ReadOnly);
 
-			for (FBackgroundHttpRequestPtr Request : ActiveRequests)
+			for (const FBackgroundHttpRequestPtr& Request : ActiveRequests)
 			{
-				if (HasUnderlyingJavaCompletedRequest(Request))
+				FAndroidBackgroundHttpRequestPtr AndroidRequest = StaticCastSharedPtr<FAndroidPlatformBackgroundHttpRequest>(Request);
+				if (HasUnderlyingJavaCompletedRequest(AndroidRequest))
 				{
-					RequestsToComplete.Add(Request);
+					RequestsToComplete.Add(MoveTemp(AndroidRequest));
 				}
 			}	
 		}
 
 		//Actually complete requests. The CompleteWithExistingResponseData call should end up removing requests from our ActiveRequests list
-		if (RequestsToComplete.Num() > 0)
+		for (const FAndroidBackgroundHttpRequestPtr& Request : RequestsToComplete)
 		{
-			for (FBackgroundHttpRequestPtr Request : RequestsToComplete)
-			{
-				FString ExistingFilePath;
-				int64 ExistingFileSize;
-				const bool bDoesCompleteFileExist = CheckForExistingCompletedDownload(Request, ExistingFilePath, ExistingFileSize);
+			Request->SendDownloadProgressUpdate();
 
-				EHttpResponseCodes::Type ResponseCodeToUse = bDoesCompleteFileExist ? EHttpResponseCodes::Ok : EHttpResponseCodes::Unknown;
-				FBackgroundHttpResponsePtr NewResponseWithExistingFile = FPlatformBackgroundHttp::ConstructBackgroundResponse(ResponseCodeToUse, ExistingFilePath);
-				Request->CompleteWithExistingResponseData(NewResponseWithExistingFile);
-			}
+			FString ExistingFilePath;
+			int64 ExistingFileSize;
+			const bool bDoesCompleteFileExist = CheckForExistingCompletedDownload(Request, ExistingFilePath, ExistingFileSize);
+
+			EHttpResponseCodes::Type ResponseCodeToUse = bDoesCompleteFileExist ? EHttpResponseCodes::Ok : EHttpResponseCodes::Unknown;
+			FBackgroundHttpResponsePtr NewResponseWithExistingFile = FPlatformBackgroundHttp::ConstructBackgroundResponse(ResponseCodeToUse, ExistingFilePath);
+			Request->CompleteWithExistingResponseData(NewResponseWithExistingFile);
 		}
 	}
 }
@@ -507,7 +507,7 @@ void FAndroidPlatformBackgroundHttpManager::HandlePendingCompletes()
 void FAndroidPlatformBackgroundHttpManager::UpdateRequestProgress()
 {
 	FRWScopeLock ScopeLock(ActiveRequestLock, SLT_ReadOnly);
-	for (FBackgroundHttpRequestPtr& Request : ActiveRequests)
+	for (const FBackgroundHttpRequestPtr& Request : ActiveRequests)
 	{
 		FAndroidBackgroundHttpRequestPtr AndroidRequest = StaticCastSharedPtr<FAndroidPlatformBackgroundHttpRequest>(Request);
 		if (ensureAlwaysMsgf(AndroidRequest.IsValid(), TEXT("Invalid request in Active Requests!")))
@@ -517,29 +517,26 @@ void FAndroidPlatformBackgroundHttpManager::UpdateRequestProgress()
 	}
 }
 
-FAndroidBackgroundHttpRequestPtr FAndroidPlatformBackgroundHttpManager::FindRequestByID(FString RequestID)
+FAndroidBackgroundHttpRequestPtr FAndroidPlatformBackgroundHttpManager::FindRequestByDestinationLocation(const FString& DestinationLocation)
 {
-	FAndroidBackgroundHttpRequestPtr ReturnedPtr = nullptr;
-
 	FRWScopeLock ScopeLock(ActiveRequestLock, SLT_ReadOnly);
-	for (FBackgroundHttpRequestPtr& Request : ActiveRequests)
+	for (const FBackgroundHttpRequestPtr& Request : ActiveRequests)
 	{
-		if (Request->GetRequestID().Equals(RequestID))
+        FAndroidBackgroundHttpRequestPtr AndroidRequest = StaticCastSharedPtr<FAndroidPlatformBackgroundHttpRequest>(Request);
+		if (AndroidRequest->DestinationLocation.Equals(DestinationLocation))
 		{
-			ReturnedPtr = StaticCastSharedPtr<FAndroidPlatformBackgroundHttpRequest>(Request);
-			break;
+			return AndroidRequest;
 		}
 	}
 
-	return ReturnedPtr;
+	return nullptr;
 }
 
 void FAndroidPlatformBackgroundHttpManager::PauseRequest(FBackgroundHttpRequestPtr Request)
 {
 	if (ensureAlwaysMsgf(IsInGameThread(), TEXT("Should only ever call PauseRequest from GameThread! Can not pause!")))
 	{
-		const FString RequestID = Request->GetRequestID();
-		UE_LOG(LogBackgroundHttpManager, Verbose, TEXT("PauseRequest called on %s"), *RequestID);
+		UE_LOG(LogBackgroundHttpManager, Verbose, TEXT("PauseRequest called on %s"), *Request->GetRequestID());
 
 		//Flag our request as paused
 		FAndroidBackgroundHttpRequestPtr AndroidRequest = StaticCastSharedPtr<FAndroidPlatformBackgroundHttpRequest>(Request);
@@ -552,7 +549,7 @@ void FAndroidPlatformBackgroundHttpManager::PauseRequest(FBackgroundHttpRequestP
 		const bool bIsThreadsafeToHandle = !FPlatformAtomics::InterlockedExchange(&bIsModifyingPauseList, true);
 		if (bIsThreadsafeToHandle)
 		{
-			RequestsToPauseByID_Java.Add(RequestID);
+			RequestsToPauseByID_Java.Add(AndroidRequest->DestinationLocation);
 
 			const bool bSanityCheck = FPlatformAtomics::InterlockedExchange(&bIsModifyingPauseList, false);
 			ensureAlwaysMsgf(bSanityCheck, TEXT("Error in bIsModifyingPauseList! Was false when we expected it to still be true from our lock!"));
@@ -560,7 +557,7 @@ void FAndroidPlatformBackgroundHttpManager::PauseRequest(FBackgroundHttpRequestP
 		//If its not safe cache off this pause so we can do it the next time it's safe to access the list
 		else
 		{
-			RequestsToPauseByID_GT.Add(RequestID);
+			RequestsToPauseByID_GT.Add(AndroidRequest->DestinationLocation);
 		}
 	}
 }
@@ -569,8 +566,7 @@ void FAndroidPlatformBackgroundHttpManager::ResumeRequest(FBackgroundHttpRequest
 {
 	if (ensureAlwaysMsgf(IsInGameThread(), TEXT("Should only ever call ResumeRequest from GameThread! Can not Resume!")))
 	{
-		const FString RequestID = Request->GetRequestID();
-		UE_LOG(LogBackgroundHttpManager, Verbose, TEXT("ResumeRequest called on %s"), *RequestID);
+		UE_LOG(LogBackgroundHttpManager, Verbose, TEXT("ResumeRequest called on %s"), *Request->GetRequestID());
 
 		//Flag our request as resumed
 		FAndroidBackgroundHttpRequestPtr AndroidRequest = StaticCastSharedPtr<FAndroidPlatformBackgroundHttpRequest>(Request);
@@ -583,7 +579,7 @@ void FAndroidPlatformBackgroundHttpManager::ResumeRequest(FBackgroundHttpRequest
 		const bool bIsThreadsafeToHandle = !FPlatformAtomics::InterlockedExchange(&bIsModifyingResumeList, true);
 		if (bIsThreadsafeToHandle)
 		{
-			RequestsToResumeByID_Java.Add(RequestID);
+			RequestsToResumeByID_Java.Add(AndroidRequest->DestinationLocation);
 
 			const bool bSanityCheck = FPlatformAtomics::InterlockedExchange(&bIsModifyingResumeList, false);
 			ensureAlwaysMsgf(bSanityCheck, TEXT("Error in bIsModifyingResumeList! Was false when we expected it to still be true from our lock!"));
@@ -591,7 +587,7 @@ void FAndroidPlatformBackgroundHttpManager::ResumeRequest(FBackgroundHttpRequest
 		//If its not safe cache off this pause so we can do it the next time it's safe to access the list
 		else
 		{
-			RequestsToResumeByID_GT.Add(RequestID);
+			RequestsToResumeByID_GT.Add(AndroidRequest->DestinationLocation);
 		}
 	}
 }
@@ -600,17 +596,17 @@ void FAndroidPlatformBackgroundHttpManager::CancelRequest(FBackgroundHttpRequest
 {
 	if (ensureAlwaysMsgf(IsInGameThread(), TEXT("Should only ever call CancelRequest from GameThread! Can not Cancel!")))
 	{
-		const FString RequestID = Request->GetRequestID();
-		UE_LOG(LogBackgroundHttpManager, Verbose, TEXT("CancelRequest called on %s"), *RequestID);
+		UE_LOG(LogBackgroundHttpManager, Verbose, TEXT("CancelRequest called on %s"), *Request->GetRequestID());
 
-		//Go ahead and remove our request since its cancelled
-		RemoveRequest(Request);
+        //Go ahead and remove our request since its cancelled
+        FAndroidBackgroundHttpRequestPtr AndroidRequest = StaticCastSharedPtr<FAndroidPlatformBackgroundHttpRequest>(Request);
+        RemoveRequest(AndroidRequest);
 
 		//Check if the Java list is safe to modify right now, and if so go ahead and add directly to it
 		const bool bIsThreadsafeToHandle = !FPlatformAtomics::InterlockedExchange(&bIsModifyingCancelList, true);
 		if (bIsThreadsafeToHandle)
 		{
-			RequestsToCancelByID_Java.Add(RequestID);
+			RequestsToCancelByID_Java.Add(AndroidRequest->DestinationLocation);
 
 			const bool bSanityCheck = FPlatformAtomics::InterlockedExchange(&bIsModifyingCancelList, false);
 			ensureAlwaysMsgf(bSanityCheck, TEXT("Error in bIsModifyingCancelList! Was false when we expected it to still be true from our lock!"));
@@ -618,7 +614,7 @@ void FAndroidPlatformBackgroundHttpManager::CancelRequest(FBackgroundHttpRequest
 		//If its not safe cache off this pause so we can do it the next time it's safe to access the list
 		else
 		{
-			RequestsToCancelByID_GT.Add(RequestID);
+			RequestsToCancelByID_GT.Add(AndroidRequest->DestinationLocation);
 		}
 	}
 }
@@ -631,9 +627,9 @@ void FAndroidPlatformBackgroundHttpManager::HandleRequestsWaitingOnJavaThreadsaf
 		const bool bIsThreadsafeToHandle = !FPlatformAtomics::InterlockedExchange(&bIsModifyingPauseList, true);
 		if (bIsThreadsafeToHandle)
 		{
-			for (const FString& RequestID : RequestsToPauseByID_GT)
+			for (FString& DestinationLocation : RequestsToPauseByID_GT)
 			{
-				RequestsToPauseByID_Java.Add(RequestID);
+				RequestsToPauseByID_Java.Add(MoveTemp(DestinationLocation));
 			}
 
 			RequestsToPauseByID_GT.Empty();
@@ -649,9 +645,9 @@ void FAndroidPlatformBackgroundHttpManager::HandleRequestsWaitingOnJavaThreadsaf
 		const bool bIsThreadsafeToHandle = !FPlatformAtomics::InterlockedExchange(&bIsModifyingResumeList, true);
 		if (bIsThreadsafeToHandle)
 		{
-			for (const FString& RequestID : RequestsToResumeByID_GT)
+			for (FString& DestinationLocation : RequestsToResumeByID_GT)
 			{
-				RequestsToResumeByID_Java.Add(RequestID);
+				RequestsToResumeByID_Java.Add(MoveTemp(DestinationLocation));
 			}
 
 			RequestsToResumeByID_GT.Empty();
@@ -667,9 +663,9 @@ void FAndroidPlatformBackgroundHttpManager::HandleRequestsWaitingOnJavaThreadsaf
 		const bool bIsThreadsafeToHandle = !FPlatformAtomics::InterlockedExchange(&bIsModifyingCancelList, true);
 		if (bIsThreadsafeToHandle)
 		{
-			for (const FString& RequestID : RequestsToCancelByID_GT)
+			for (FString& DestinationLocation : RequestsToCancelByID_GT)
 			{
-				RequestsToCancelByID_Java.Add(RequestID);
+				RequestsToCancelByID_Java.Add(MoveTemp(DestinationLocation));
 			}
 
 			RequestsToCancelByID_GT.Empty();
@@ -692,38 +688,38 @@ const FString FAndroidPlatformBackgroundHttpManager::GetFullFileNameForDownloadD
 	return FPaths::Combine(BGDownloadDescriptionFolder, FileName);
 }
 
-void FAndroidPlatformBackgroundHttpManager::Java_OnDownloadProgress(jobject UnderlyingWorker, FString RequestID, int64_t BytesWrittenSinceLastCall, int64_t TotalBytesWritten)
+void FAndroidPlatformBackgroundHttpManager::Java_OnDownloadProgress(jobject UnderlyingWorker, FString DestinationLocation, int64_t TotalBytesWritten, int64_t BytesWrittenSinceLastCall)
 {
-	UE_LOG(LogBackgroundHttpManager, VeryVerbose, TEXT("Download Progress... RequestID:%s BytesWrittenSinceLastCall:%lld TotalBytesWritten:%lld"), *RequestID, BytesWrittenSinceLastCall, TotalBytesWritten);
+	UE_LOG(LogBackgroundHttpManager, VeryVerbose, TEXT("Download Progress... DestinationLocation:%s BytesWrittenSinceLastCall:%lld TotalBytesWritten:%lld"), *DestinationLocation, BytesWrittenSinceLastCall, TotalBytesWritten);
 
-	FAndroidBackgroundHttpRequestPtr FoundRequest = FindRequestByID(RequestID);
+	FAndroidBackgroundHttpRequestPtr FoundRequest = FindRequestByDestinationLocation(DestinationLocation);
 	if (FoundRequest.IsValid())
 	{
 		FoundRequest->UpdateDownloadProgress(TotalBytesWritten, BytesWrittenSinceLastCall);
 	}
 }
 
-void FAndroidPlatformBackgroundHttpManager::Java_OnDownloadComplete(jobject UnderlyingWorker, FString RequestID, FString CompleteLocation, bool bWasSuccess)
+void FAndroidPlatformBackgroundHttpManager::Java_OnDownloadComplete(jobject UnderlyingWorker, FString DestinationLocation, bool bWasSuccess)
 {
-	UE_LOG(LogBackgroundHttpManager, Log, TEXT("DownloadComplete... RequestID:%s bWasSuccess:%d"), *RequestID, (int)bWasSuccess);
+	UE_LOG(LogBackgroundHttpManager, Log, TEXT("DownloadComplete... DestinationLocation:%s bWasSuccess:%d"), *DestinationLocation, (int)bWasSuccess);
 
 	//Mark associated request as completed so that we can complete it on our next tick
-	FAndroidBackgroundHttpRequestPtr CompletedRequest = FindRequestByID(RequestID);
+	FAndroidBackgroundHttpRequestPtr CompletedRequest = FindRequestByDestinationLocation(DestinationLocation);
 	if (CompletedRequest.IsValid())
 	{
 		MarkUnderlyingJavaRequestAsCompleted(CompletedRequest);
 	}
 	else
 	{
-		UE_LOG(LogBackgroundHttpManager, Log, TEXT("Taking no action as RequestID:%s did not have a corresponding ActiveRequest"), *RequestID);
+		UE_LOG(LogBackgroundHttpManager, Log, TEXT("Taking no action as DestinationLocation:%s did not have a corresponding ActiveRequest"), *DestinationLocation);
 	}
 }
 
-void FAndroidPlatformBackgroundHttpManager::Java_OnDownloadMetrics(jobject UnderlyingWorker, FString RequestID, int64_t TotalBytesDownloaded, int64_t DownloadDuration, int64_t DownloadStartTimeUTC, int64_t DownloadEndTimeUTC)
+void FAndroidPlatformBackgroundHttpManager::Java_OnDownloadMetrics(jobject UnderlyingWorker, FString DestinationLocation, int64_t TotalBytesDownloaded, int64_t DownloadDuration, int64_t DownloadStartTimeUTC, int64_t DownloadEndTimeUTC)
 {
-	UE_LOG(LogBackgroundHttpManager, Log, TEXT("DownloadMetrics... RequestID:%s TotalBytesDownloaded:%ld DownloadDuration:%ld"), *RequestID, TotalBytesDownloaded, DownloadDuration);
+	UE_LOG(LogBackgroundHttpManager, Log, TEXT("DownloadMetrics... DestinationLocation:%s TotalBytesDownloaded:%ld DownloadDuration:%ld"), *DestinationLocation, TotalBytesDownloaded, DownloadDuration);
 
-	FAndroidBackgroundHttpRequestPtr CompletedRequest = FindRequestByID(RequestID);
+	FAndroidBackgroundHttpRequestPtr CompletedRequest = FindRequestByDestinationLocation(DestinationLocation);
 	if (CompletedRequest.IsValid())
 	{
 		FBackgroundHttpRequestMetricsExtended MetricsExtended;
@@ -885,13 +881,9 @@ void FAndroidPlatformBackgroundHttpManager::Java_OnTick(JNIEnv* Env, jobject Und
 					jmethodID PauseMethod = Env->GetMethodID(UnderlyingWorkerClass, "PauseRequest", "(Ljava/lang/String;)V");
 					CHECK_JNI_RESULT(PauseMethod);
 
-					for (const FString& RequestID : RequestsToPauseByID_Java)
+					for (const FString& DestinationLocation : RequestsToPauseByID_Java)
 					{
-						//Get java version of our RequestID
-						FScopedJavaObject<jstring> ConvertedRequestID = FJavaHelper::ToJavaString(Env, RequestID);
-
-						//Call pause method on RequestID
-						FJavaWrapper::CallVoidMethod(Env, UnderlyingWorker, PauseMethod, *ConvertedRequestID);
+						FJavaWrapper::CallVoidMethod(Env, UnderlyingWorker, PauseMethod, *FJavaHelper::ToJavaString(Env, DestinationLocation));
 					}
 
 					RequestsToPauseByID_Java.Empty();
@@ -910,13 +902,9 @@ void FAndroidPlatformBackgroundHttpManager::Java_OnTick(JNIEnv* Env, jobject Und
 					jmethodID ResumeMethod = Env->GetMethodID(UnderlyingWorkerClass, "ResumeRequest", "(Ljava/lang/String;)V");
 					CHECK_JNI_RESULT(ResumeMethod);
 
-					for (const FString& RequestID : RequestsToResumeByID_Java)
+					for (const FString& DestinationLocation : RequestsToResumeByID_Java)
 					{
-						//Get java version of our RequestID
-						FScopedJavaObject<jstring> ConvertedRequestID = FJavaHelper::ToJavaString(Env, RequestID);
-
-						//Call pause method on RequestID
-						FJavaWrapper::CallVoidMethod(Env, UnderlyingWorker, ResumeMethod, *ConvertedRequestID);
+						FJavaWrapper::CallVoidMethod(Env, UnderlyingWorker, ResumeMethod, *FJavaHelper::ToJavaString(Env, DestinationLocation));
 					}
 
 					RequestsToResumeByID_Java.Empty();
@@ -934,13 +922,9 @@ void FAndroidPlatformBackgroundHttpManager::Java_OnTick(JNIEnv* Env, jobject Und
 					jmethodID CancelMethod = Env->GetMethodID(UnderlyingWorkerClass, "CancelRequest", "(Ljava/lang/String;)V");
 					CHECK_JNI_RESULT(CancelMethod);
 
-					for (const FString& RequestID : RequestsToCancelByID_Java)
+					for (const FString& DestinationLocation : RequestsToCancelByID_Java)
 					{
-						//Get java version of our RequestID
-						FScopedJavaObject<jstring> ConvertedRequestID = FJavaHelper::ToJavaString(Env, RequestID);
-
-						// Call cancel method on RequestID
-						FJavaWrapper::CallVoidMethod(Env, UnderlyingWorker, CancelMethod, *ConvertedRequestID);
+						FJavaWrapper::CallVoidMethod(Env, UnderlyingWorker, CancelMethod, *FJavaHelper::ToJavaString(Env, DestinationLocation));
 					}
 
 					RequestsToCancelByID_Java.Empty();
@@ -956,7 +940,7 @@ void FAndroidPlatformBackgroundHttpManager::Java_OnTick(JNIEnv* Env, jobject Und
 void FAndroidPlatformBackgroundHttpManager::ForceCompleteAllUnderlyingJavaActiveRequestsWithError()
 {
 	FRWScopeLock ScopeLock(ActiveRequestLock, SLT_ReadOnly);
-	for (FBackgroundHttpRequestPtr& Request : ActiveRequests)
+	for (const FBackgroundHttpRequestPtr& Request : ActiveRequests)
 	{
 		if (!HasUnderlyingJavaCompletedRequest(Request))
 		{
@@ -1265,33 +1249,32 @@ void JNICALL UE::Jni::Download::FUEDownloadWorker::nativeAndroidBackgroundDownlo
 	}
 }
 
-void JNICALL UE::Jni::Download::FUEDownloadWorker::nativeAndroidBackgroundDownloadOnProgress(JNIEnv* jenv, jobject thiz, jstring TaskID, jlong BytesWrittenSinceLastCall, jlong TotalBytesWritten)
+void JNICALL UE::Jni::Download::FUEDownloadWorker::nativeAndroidBackgroundDownloadOnProgress(JNIEnv* jenv, jobject thiz, jstring TaskID, jlong TotalBytesWritten, jlong BytesWrittenSinceLastCall)
 {
-	FString RequestID = FJavaHelper::FStringFromParam(jenv, TaskID);
-	int64_t ConvertedBytesWrittenSinceLastCall = static_cast<uint64_t>(BytesWrittenSinceLastCall);
+	FString DestinationLocation = FJavaHelper::FStringFromParam(jenv, TaskID);
 	int64_t ConvertedTotalBytesWritten = static_cast<uint64_t>(TotalBytesWritten);
+	int64_t ConvertedBytesWrittenSinceLastCall = static_cast<uint64_t>(BytesWrittenSinceLastCall);
 
-	FAndroidBackgroundDownloadDelegates::AndroidBackgroundDownload_OnProgress.ExecuteIfBound(thiz, RequestID, ConvertedBytesWrittenSinceLastCall, ConvertedTotalBytesWritten);
+	FAndroidBackgroundDownloadDelegates::AndroidBackgroundDownload_OnProgress.ExecuteIfBound(thiz, MoveTemp(DestinationLocation), ConvertedTotalBytesWritten, ConvertedBytesWrittenSinceLastCall);
 }
 
-void JNICALL UE::Jni::Download::FUEDownloadWorker::nativeAndroidBackgroundDownloadOnComplete(JNIEnv* jenv, jobject thiz, jstring TaskID, jstring CompleteLocation, jboolean bWasSuccess)
+void JNICALL UE::Jni::Download::FUEDownloadWorker::nativeAndroidBackgroundDownloadOnComplete(JNIEnv* jenv, jobject thiz, jstring TaskID, jboolean bWasSuccess)
 {
-	FString RequestID = FJavaHelper::FStringFromParam(jenv, TaskID);
-	FString ConvertedCompleteLocation = FJavaHelper::FStringFromParam(jenv, CompleteLocation);
+	FString DestinationLocation = FJavaHelper::FStringFromParam(jenv, TaskID);
 	bool ConvertedbWasSuccess = static_cast<bool>(bWasSuccess);
 
-	FAndroidBackgroundDownloadDelegates::AndroidBackgroundDownload_OnComplete.ExecuteIfBound(thiz, RequestID, ConvertedCompleteLocation, ConvertedbWasSuccess);
+	FAndroidBackgroundDownloadDelegates::AndroidBackgroundDownload_OnComplete.ExecuteIfBound(thiz, MoveTemp(DestinationLocation), ConvertedbWasSuccess);
 }
 
 void JNICALL UE::Jni::Download::FUEDownloadWorker::nativeAndroidBackgroundDownloadOnMetrics(JNIEnv* jenv, jobject thiz, jstring TaskID, jlong TotalBytesDownloaded, jlong DownloadDuration, jlong DownloadStartTimeUTC, jlong DownloadEndTimeUTC)
 {
-	FString RequestID = FJavaHelper::FStringFromParam(jenv, TaskID);
+	FString DestinationLocation = FJavaHelper::FStringFromParam(jenv, TaskID);
 	int64_t ConvertedTotalBytesDownloaded = static_cast<uint64_t>(TotalBytesDownloaded);
 	int64_t ConvertedDownloadDuration = static_cast<uint64_t>(DownloadDuration);
 	int64_t ConvertedDownloadStartTimeUTC = static_cast<uint64_t>(DownloadStartTimeUTC);
 	int64_t ConvertedDownloadEndTimeUTC = static_cast<uint64_t>(DownloadEndTimeUTC);
 
-	FAndroidBackgroundDownloadDelegates::AndroidBackgroundDownload_OnMetrics.ExecuteIfBound(thiz, RequestID, ConvertedTotalBytesDownloaded, ConvertedDownloadDuration, ConvertedDownloadStartTimeUTC, ConvertedDownloadEndTimeUTC);
+	FAndroidBackgroundDownloadDelegates::AndroidBackgroundDownload_OnMetrics.ExecuteIfBound(thiz, MoveTemp(DestinationLocation), ConvertedTotalBytesDownloaded, ConvertedDownloadDuration, ConvertedDownloadStartTimeUTC, ConvertedDownloadEndTimeUTC);
 }
 
 void JNICALL UE::Jni::Download::FUEDownloadWorker::nativeAndroidBackgroundDownloadOnAllComplete(JNIEnv* jenv, jobject thiz, jboolean bDidAllRequestsSucceed)
diff --git a/Engine/Plugins/Online/Android/AndroidFetchBackgroundDownload/Source/AndroidFetchBackgroundDownload/Private/AndroidPlatformBackgroundHttpRequest.cpp b/Engine/Plugins/Online/Android/AndroidFetchBackgroundDownload/Source/AndroidFetchBackgroundDownload/Private/AndroidPlatformBackgroundHttpRequest.cpp
--- a/Engine/Plugins/Online/Android/AndroidFetchBackgroundDownload/Source/AndroidFetchBackgroundDownload/Private/AndroidPlatformBackgroundHttpRequest.cpp
+++ b/Engine/Plugins/Online/Android/AndroidFetchBackgroundDownload/Source/AndroidFetchBackgroundDownload/Private/AndroidPlatformBackgroundHttpRequest.cpp
@@ -11,16 +11,16 @@
 typedef  TJsonWriter<TCHAR, TCondensedJsonPrintPolicy<TCHAR>> FAndroidHttpRequestJsonWriter;
 typedef  TJsonWriterFactory<TCHAR, TCondensedJsonPrintPolicy<TCHAR>> FAndroidHttpRequestJsonWriterFactory;
 
-const FString FAndroidPlatformBackgroundHttpRequest::URLKey = "URLs";
-const FString FAndroidPlatformBackgroundHttpRequest::RequestIDKey = "RequestID";
-const FString FAndroidPlatformBackgroundHttpRequest::DestinationLocationKey = "DestLocation";
-const FString FAndroidPlatformBackgroundHttpRequest::MaxRetryCountKey = "MaxRetryCount";
-const FString FAndroidPlatformBackgroundHttpRequest::IndividualURLRetryCountKey = "IndividualURLRetryCount";
-const FString FAndroidPlatformBackgroundHttpRequest::RequestPriorityKey = "RequestPriority";
-const FString FAndroidPlatformBackgroundHttpRequest::GroupIDKey = "GroupId";
-const FString FAndroidPlatformBackgroundHttpRequest::bHasCompletedKey = "bHasCompleted";
-const FString FAndroidPlatformBackgroundHttpRequest::bIsPausedKey = "bIsPaused";
-const FString FAndroidPlatformBackgroundHttpRequest::TotalBytesNeededKey = "TotalBytesNeeded";
+const FStringView FAndroidPlatformBackgroundHttpRequest::URLKey = TEXTVIEW("URLs");
+const FStringView FAndroidPlatformBackgroundHttpRequest::RequestIDKey = TEXTVIEW("RequestID");
+const FStringView FAndroidPlatformBackgroundHttpRequest::DestinationLocationKey = TEXTVIEW("DestLocation");
+const FStringView FAndroidPlatformBackgroundHttpRequest::MaxRetryCountKey = TEXTVIEW("MaxRetryCount");
+const FStringView FAndroidPlatformBackgroundHttpRequest::IndividualURLRetryCountKey = TEXTVIEW("IndividualURLRetryCount");
+const FStringView FAndroidPlatformBackgroundHttpRequest::RequestPriorityKey = TEXTVIEW("RequestPriority");
+const FStringView FAndroidPlatformBackgroundHttpRequest::GroupIDKey = TEXTVIEW("GroupId");
+const FStringView FAndroidPlatformBackgroundHttpRequest::bHasCompletedKey = TEXTVIEW("bHasCompleted");
+const FStringView FAndroidPlatformBackgroundHttpRequest::bIsPausedKey = TEXTVIEW("bIsPaused");
+const FStringView FAndroidPlatformBackgroundHttpRequest::TotalBytesNeededKey = TEXTVIEW("TotalBytesNeeded");
 
 FAndroidPlatformBackgroundHttpRequest::FAndroidPlatformBackgroundHttpRequest()
 	: bIsPaused(false)
@@ -59,6 +59,8 @@ void FAndroidPlatformBackgroundHttpRequest::ResumeRequest()
 
 FString FAndroidPlatformBackgroundHttpRequest::ToJSon() const
 {
+    DestinationLocation = GetDestinationLocation();
+    
 	FString JsonOutput;
 	TSharedRef<FAndroidHttpRequestJsonWriter> Writer = FAndroidHttpRequestJsonWriterFactory::Create(&JsonOutput);
 	Writer->WriteObjectStart();
@@ -66,7 +68,7 @@ FString FAndroidPlatformBackgroundHttpRequest::ToJSon() const
 		Writer->WriteValue(RequestIDKey, RequestID);
 		Writer->WriteValue(RequestPriorityKey, GetPriorityAsAndroidPriority());
 		Writer->WriteValue(MaxRetryCountKey, NumberOfTotalRetries);
-		Writer->WriteValue(DestinationLocationKey, GetDestinationLocation());
+		Writer->WriteValue(DestinationLocationKey, DestinationLocation);
 		
 		//Write this bool as either true/false so it can be parsed in JSON as a java Boolean object
 		const bool bIsCompletedCopy = FPlatformAtomics::AtomicRead(&bIsCompleted);
@@ -152,13 +154,13 @@ void FAndroidPlatformBackgroundHttpRequest::SendDownloadProgressUpdate()
 	//so that requestors don't have to worry about thread safety unexpectidly
 	ensureAlwaysMsgf(IsInGameThread(), TEXT("Called from un-expected thread! Potential error in an implementation of background downloads!"));
 
-	volatile int64 DownloadProgressCopy = FPlatformAtomics::AtomicRead(&DownloadProgress);
-
+	//Reset our DownloadProgressSinceLastUpdateSent to 0 now that we are sending a progress update
+	int64 DownloadProgressSinceLastUpdateSentCopy = FPlatformAtomics::InterlockedExchange(&DownloadProgressSinceLastUpdateSent, 0);
+	
 	//Don't send any updates if we haven't updated anything since we last sent an update
-	if (DownloadProgressCopy > 0)
+	if (DownloadProgressSinceLastUpdateSentCopy > 0)
 	{
-		//Reset our DownloadProgressSinceLastUpdateSent to 0 now that we are sending a progress update
-		volatile int64 DownloadProgressSinceLastUpdateSentCopy = FPlatformAtomics::InterlockedExchange(&DownloadProgressSinceLastUpdateSent, 0);
+		int64 DownloadProgressCopy = FPlatformAtomics::AtomicRead(&DownloadProgress);
 
 		OnProgressUpdated().ExecuteIfBound(SharedThis(this), DownloadProgressCopy, DownloadProgressSinceLastUpdateSentCopy);
 	}
diff --git a/Engine/Plugins/Online/Android/AndroidFetchBackgroundDownload/Source/AndroidFetchBackgroundDownload/Public/AndroidPlatformBackgroundHttpManager.h b/Engine/Plugins/Online/Android/AndroidFetchBackgroundDownload/Source/AndroidFetchBackgroundDownload/Public/AndroidPlatformBackgroundHttpManager.h
--- a/Engine/Plugins/Online/Android/AndroidFetchBackgroundDownload/Source/AndroidFetchBackgroundDownload/Public/AndroidPlatformBackgroundHttpManager.h
+++ b/Engine/Plugins/Online/Android/AndroidFetchBackgroundDownload/Source/AndroidFetchBackgroundDownload/Public/AndroidPlatformBackgroundHttpManager.h
@@ -63,7 +63,7 @@ public:
 protected:
 	void UpdateRequestProgress();
 
-	FAndroidBackgroundHttpRequestPtr FindRequestByID(FString RequestID);
+	FAndroidBackgroundHttpRequestPtr FindRequestByDestinationLocation(const FString& DestinationLocation);
 	void HandlePendingCompletes();
 
 	void HandleJavaWorkerReconciliation();
@@ -73,9 +73,9 @@ protected:
 	//Handlers for our download progressing in the underlying java implementation so that we can bubble it up to UE code.
 	void Java_OnWorkerStart(FString WorkID, jobject UnderlyingWorker);
 	void Java_OnWorkerStop(FString WorkID, jobject UnderlyingWorker);
-	void Java_OnDownloadProgress(jobject UnderlyingWorker, FString RequestID, int64_t BytesWrittenSinceLastCall, int64_t TotalBytesWritten);
-	void Java_OnDownloadComplete(jobject UnderlyingWorker, FString RequestID, FString CompleteLocation, bool bWasSuccess);
-	void Java_OnDownloadMetrics(jobject UnderlyingWorker, FString RequestID, int64_t TotalBytesDownloaded, int64_t DownloadDuration, int64_t DownloadStartTimeUTC, int64_t DownloadEndTimeUTC);
+	void Java_OnDownloadProgress(jobject UnderlyingWorker, FString DestinationLocation, int64_t TotalBytesWritten, int64_t BytesWrittenSinceLastCall);
+	void Java_OnDownloadComplete(jobject UnderlyingWorker, FString DestinationLocation, bool bWasSuccess);
+	void Java_OnDownloadMetrics(jobject UnderlyingWorker, FString DestinationLocation, int64_t TotalBytesDownloaded, int64_t DownloadDuration, int64_t DownloadStartTimeUTC, int64_t DownloadEndTimeUTC);
 	void Java_OnAllDownloadsComplete(jobject UnderlyingWorker, bool bDidAllRequestsSucceed);
 	void Java_OnTick(JNIEnv* Env, jobject UnderlyingWorker);
 
diff --git a/Engine/Plugins/Online/Android/AndroidFetchBackgroundDownload/Source/AndroidFetchBackgroundDownload/Public/AndroidPlatformBackgroundHttpRequest.h b/Engine/Plugins/Online/Android/AndroidFetchBackgroundDownload/Source/AndroidFetchBackgroundDownload/Public/AndroidPlatformBackgroundHttpRequest.h
--- a/Engine/Plugins/Online/Android/AndroidFetchBackgroundDownload/Source/AndroidFetchBackgroundDownload/Public/AndroidPlatformBackgroundHttpRequest.h
+++ b/Engine/Plugins/Online/Android/AndroidFetchBackgroundDownload/Source/AndroidFetchBackgroundDownload/Public/AndroidPlatformBackgroundHttpRequest.h
@@ -34,18 +34,19 @@ private:
 private:
 	//NOTE: These must match the keys used in the DownloadDescription.java file in The AndroidBackgoundService module!
 	//We use these keys to convert our request into a DownloadDescription on the java side through JSON serialization!
-	static const FString URLKey;
-	static const FString RequestIDKey;
-	static const FString DestinationLocationKey;
-	static const FString MaxRetryCountKey;
-	static const FString IndividualURLRetryCountKey;
-	static const FString RequestPriorityKey;
-	static const FString GroupIDKey;
-	static const FString bHasCompletedKey;
-	static const FString bIsPausedKey;
-	static const FString TotalBytesNeededKey;
+	static const FStringView URLKey;
+	static const FStringView RequestIDKey;
+	static const FStringView DestinationLocationKey;
+	static const FStringView MaxRetryCountKey;
+	static const FStringView IndividualURLRetryCountKey;
+	static const FStringView RequestPriorityKey;
+	static const FStringView GroupIDKey;
+	static const FStringView bHasCompletedKey;
+	static const FStringView bIsPausedKey;
+	static const FStringView TotalBytesNeededKey;
 
 private:
+    mutable FString DestinationLocation;
 	volatile int32 bIsPaused;
 	volatile int32 bIsCompleted;
 	volatile int64 DownloadProgress;
diff --git a/Engine/Source/Runtime/Online/BackgroundHTTP/Private/BackgroundHttpManagerImpl.cpp b/Engine/Source/Runtime/Online/BackgroundHTTP/Private/BackgroundHttpManagerImpl.cpp
--- a/Engine/Source/Runtime/Online/BackgroundHTTP/Private/BackgroundHttpManagerImpl.cpp
+++ b/Engine/Source/Runtime/Online/BackgroundHTTP/Private/BackgroundHttpManagerImpl.cpp
@@ -232,7 +232,7 @@ void FBackgroundHttpManagerImpl::GatherAllTempFilenames(TArray<FString>& OutAllT
 {
 	OutAllTempFilenames.Empty();
 	
-	const FString DirectoryToCheck = GetFileHashHelper()->GetTemporaryRootPath();
+	const FString& DirectoryToCheck = FBackgroundHttpFileHashHelper::GetTemporaryDownloadPath();
 	
 	TArray<FString> AllFilenames;
 	IFileManager::Get().FindFiles(AllFilenames, *DirectoryToCheck, *FBackgroundHttpFileHashHelper::GetTempFileExtension());
@@ -272,7 +272,7 @@ void FBackgroundHttpManagerImpl::DeleteAllTemporaryFiles()
 
 	//Default implementation is to just delete everything in the Root Folder non-recursively.
 	IPlatformFile& PlatformFile = FPlatformFileManager::Get().GetPlatformFile();
-	PlatformFile.FindFiles(FilesToDelete, *FBackgroundHttpFileHashHelper::GetTemporaryRootPath(), *FBackgroundHttpFileHashHelper::GetTempFileExtension());
+	PlatformFile.FindFiles(FilesToDelete, *FBackgroundHttpFileHashHelper::GetTemporaryDownloadPath(), *FBackgroundHttpFileHashHelper::GetTempFileExtension());
 
 	for (const FString& File : FilesToDelete)
 	{
diff --git a/Engine/Source/Runtime/Online/BackgroundHTTP/Private/BackgroundHttpRequestImpl.cpp b/Engine/Source/Runtime/Online/BackgroundHTTP/Private/BackgroundHttpRequestImpl.cpp
--- a/Engine/Source/Runtime/Online/BackgroundHTTP/Private/BackgroundHttpRequestImpl.cpp
+++ b/Engine/Source/Runtime/Online/BackgroundHTTP/Private/BackgroundHttpRequestImpl.cpp
@@ -10,12 +10,9 @@ DEFINE_LOG_CATEGORY(LogBackgroundHttpRequest);
 FBackgroundHttpRequestImpl::FBackgroundHttpRequestImpl()
 	: DownloadCompleteNotificationObject(nullptr)
 	, Response(nullptr)
-	, URLList()
-	, RequestID()
 	, NumberOfTotalRetries(0)
 	, RequestPriority(EBackgroundHTTPPriority::Normal)
-	, HttpRequestCompleteDelegate()
-	, HttpProgressUpdateDelegate()
+	, ExpectedResultSize{0}
 {
 }
 
diff --git a/Engine/Source/Runtime/Online/BackgroundHTTPFileHash/Private/BackgroundHttpFileHashHelper.cpp b/Engine/Source/Runtime/Online/BackgroundHTTPFileHash/Private/BackgroundHttpFileHashHelper.cpp
--- a/Engine/Source/Runtime/Online/BackgroundHTTPFileHash/Private/BackgroundHttpFileHashHelper.cpp
+++ b/Engine/Source/Runtime/Online/BackgroundHTTPFileHash/Private/BackgroundHttpFileHashHelper.cpp
@@ -81,6 +81,16 @@ const FString& FBackgroundHttpFileHashHelper::GetTemporaryRootPath()
 	return BackgroundHttpDir;
 }
 
+const FString& FBackgroundHttpFileHashHelper::GetTemporaryDownloadPath()
+{
+#if PLATFORM_ANDROID || PLATFORM_IOS
+    static FString DownloadDir = FPaths::Combine(GetTemporaryRootPath(), BuildSettings::GetBranchName());
+    return DownloadDir;
+#else
+    return GetTemporaryRootPath();
+#endif
+}
+
 const FString& FBackgroundHttpFileHashHelper::GetTempFileExtension()
 {
 	static FString TempExtension = ".tmp";
@@ -97,10 +107,7 @@ FString FBackgroundHttpFileHashHelper::GetFullPathOfTempFilename(const FString&
 {
 	// don't take lock here, as other methods taking lock call it
 	FString FullFileName = FPaths::SetExtension(FPaths::GetCleanFilename(TempFilename), GetTempFileExtension());
-#if PLATFORM_ANDROID || PLATFORM_IOS
-	FullFileName = FPaths::Combine(BuildSettings::GetBranchName(), FullFileName);
-#endif
-	return FPaths::Combine(GetTemporaryRootPath(), FullFileName);
+	return FPaths::Combine(GetTemporaryDownloadPath(), FullFileName);
 }
 
 FString FBackgroundHttpFileHashHelper::ToString() const
diff --git a/Engine/Source/Runtime/Online/BackgroundHTTPFileHash/Public/BackgroundHttpFileHashHelper.h b/Engine/Source/Runtime/Online/BackgroundHTTPFileHash/Public/BackgroundHttpFileHashHelper.h
--- a/Engine/Source/Runtime/Online/BackgroundHTTPFileHash/Public/BackgroundHttpFileHashHelper.h
+++ b/Engine/Source/Runtime/Online/BackgroundHTTPFileHash/Public/BackgroundHttpFileHashHelper.h
@@ -90,6 +90,9 @@ public:
 	// Get our base directory used to store Temp files
 	static BACKGROUNDHTTPFILEHASH_API const FString& GetTemporaryRootPath();
 
+    // Get temporary download path
+    static BACKGROUNDHTTPFILEHASH_API const FString& GetTemporaryDownloadPath();
+    
 	// Helper function that returns the file extension used by our BackgroundHTTP temp files
 	static BACKGROUNDHTTPFILEHASH_API const FString& GetTempFileExtension();
 
