Android patch download progress is not reported correctly due to change in the ProgressListener API
- Progress listener APIs being called with IDs instead of destination locations

#Synced-CL 49876485
#Synced-CL 49876485

diff --git a/Engine/Plugins/Online/Android/AndroidFetchBackgroundDownload/Source/AndroidFetchBackgroundDownload/Java/download/asyncdownloader/AsyncDownloadRequest.java b/Engine/Plugins/Online/Android/AndroidFetchBackgroundDownload/Source/AndroidFetchBackgroundDownload/Java/download/asyncdownloader/AsyncDownloadRequest.java
--- a/Engine/Plugins/Online/Android/AndroidFetchBackgroundDownload/Source/AndroidFetchBackgroundDownload/Java/download/asyncdownloader/AsyncDownloadRequest.java
+++ b/Engine/Plugins/Online/Android/AndroidFetchBackgroundDownload/Source/AndroidFetchBackgroundDownload/Java/download/asyncdownloader/AsyncDownloadRequest.java
@@ -51,7 +51,7 @@ final public class AsyncDownloadRequest {
 			return this; 
 		}
         public Builder SetRequestId(String RequestId) { 
-			this.Id = RequestId; 
+			this.RequestId = RequestId; 
 			return this; 
 		}
         public Builder SetUrls(List<String> Urls) { 
