Reimplemented Android Background HTTP Manager

#android

#Synced-CL 50957382

diff --git a/Engine/Build/Android/Java/src/com/epicgames/unreal/GameActivity.java.template b/Engine/Build/Android/Java/src/com/epicgames/unreal/GameActivity.java.template
--- a/Engine/Build/Android/Java/src/com/epicgames/unreal/GameActivity.java.template
+++ b/Engine/Build/Android/Java/src/com/epicgames/unreal/GameActivity.java.template
@@ -7529,7 +7529,7 @@ Trace.endSection();
 		{
 			NetworkListener = new NetworkConnectivityClient.Listener() {
 				@Override
-				public void onNetworkAvailable(NetworkConnectivityClient.NetworkTransportType networkTransportType)
+				public void onNetworkAvailable(NetworkConnectivityClient.NetworkTransportType networkTransportType, boolean isNetworkMetered)
 				{
 					nativeNetworkChanged();
 				}
diff --git a/Engine/Build/Android/Java/src/com/epicgames/unreal/network/NetworkChangedManager.java b/Engine/Build/Android/Java/src/com/epicgames/unreal/network/NetworkChangedManager.java
--- a/Engine/Build/Android/Java/src/com/epicgames/unreal/network/NetworkChangedManager.java
+++ b/Engine/Build/Android/Java/src/com/epicgames/unreal/network/NetworkChangedManager.java
@@ -2,14 +2,10 @@
 
 package com.epicgames.unreal.network;
 
-import android.annotation.TargetApi;
 import android.content.Context;
 import android.net.ConnectivityManager;
 import android.net.Network;
 import android.net.NetworkCapabilities;
-import android.net.NetworkInfo;
-import android.net.NetworkRequest;
-import android.os.Build;
 import android.os.Handler;
 import androidx.annotation.NonNull;
 import androidx.annotation.Nullable;
@@ -21,9 +17,7 @@ import java.lang.ref.WeakReference;
 import java.net.HttpURLConnection;
 import java.net.MalformedURLException;
 import java.net.URL;
-import java.util.HashSet;
-import java.util.Iterator;
-import java.util.Set;
+import java.util.concurrent.CopyOnWriteArrayList;
 import java.util.concurrent.ExecutorService;
 import java.util.concurrent.Executors;
 
@@ -47,7 +41,7 @@ public final class NetworkChangedManager implements NetworkConnectivityClient {
 
 	private static NetworkChangedManager instance;
 	@NonNull
-	public static synchronized NetworkConnectivityClient getInstance() {
+	public static synchronized NetworkChangedManager getInstance() {
 		if (instance == null) {
 			instance = new NetworkChangedManager();
 		}
@@ -60,12 +54,14 @@ public final class NetworkChangedManager implements NetworkConnectivityClient {
 	private ConnectivityManager connectivityManager;
 
 	@NonNull
-	private Set<WeakReference<Listener>> networkChangedListeners = new HashSet<>();
+	private final CopyOnWriteArrayList<WeakReference<Listener>> networkChangedListeners = new CopyOnWriteArrayList<>();
 
 	/*
 	 * Data
 	 */
-	private Set<String> networks = new HashSet<>();
+	private boolean isNetworkAvailable;
+	private boolean isNetworkMetered;
+	private boolean isNetworkBlocked;
 	@Nullable
 	private ConnectivityState currentState = null;
 	@Nullable
@@ -82,80 +78,53 @@ public final class NetworkChangedManager implements NetworkConnectivityClient {
 
 	@Override
 	public void initNetworkCallback(@NonNull Context context) {
-		if (Build.VERSION.SDK_INT < Build.VERSION_CODES.LOLLIPOP) {
-			// unsupported before Lollipop
-			return;
-		}
-
 		connectivityManager = (ConnectivityManager) context.getApplicationContext().getSystemService(Context.CONNECTIVITY_SERVICE);
 		if (connectivityManager == null) {
 			Log.error("Unable to start connectivityManager");
 			return;
 		}
-		NetworkRequest.Builder builder = new NetworkRequest.Builder();
-		builder.addCapability(NetworkCapabilities.NET_CAPABILITY_INTERNET);
-		builder.addTransportType(NetworkCapabilities.TRANSPORT_WIFI);
-		builder.addTransportType(NetworkCapabilities.TRANSPORT_ETHERNET);
-		builder.addTransportType(NetworkCapabilities.TRANSPORT_VPN);
-		builder.addTransportType(NetworkCapabilities.TRANSPORT_CELLULAR);
-		builder.addTransportType(NetworkCapabilities.TRANSPORT_BLUETOOTH);
-		if (Build.VERSION.SDK_INT >= Build.VERSION_CODES.M) {
-			builder.addCapability(NetworkCapabilities.NET_CAPABILITY_VALIDATED);
-		}
-		connectivityManager.registerNetworkCallback(builder.build(), connectivityListener);
+		connectivityManager.registerDefaultNetworkCallback(connectivityListener);
 	}
 
-	@TargetApi(Build.VERSION_CODES.LOLLIPOP)
 	private ConnectivityManager.NetworkCallback connectivityListener = new ConnectivityManager.NetworkCallback() {
-		@Override
-		public void onAvailable(Network network) {
-			networks.add(network.toString());
-			Log.verbose("Network Available: " + network.toString());
-			checkNetworkConnectivity();
-		}
-
 		@Override
 		public void onLost(Network network) {
-			networks.remove(network.toString());
+			isNetworkAvailable = false;
 			Log.verbose("Network Lost callback: " + network.toString());
-
-			if (networks.isEmpty()) {
-				Log.verbose("All Networks Lost");
-			} else {
-				Log.verbose("Not All Networks Lost");
-			}
 			checkNetworkConnectivity();
 		}
 
 		@Override
 		public void onCapabilitiesChanged(Network network, NetworkCapabilities networkCapabilities) {
-			if (Build.VERSION.SDK_INT >= Build.VERSION_CODES.M
-					&& !networkCapabilities.hasCapability(NetworkCapabilities.NET_CAPABILITY_VALIDATED)) {
+			if (!networkCapabilities.hasCapability(NetworkCapabilities.NET_CAPABILITY_VALIDATED)) {
+				isNetworkAvailable = false;
 				Log.verbose("Network Capabilities changed, doesn't have validated net_capability");
-				networks.remove(network.toString());
-			} else if (!networkCapabilities.hasCapability(NetworkCapabilities.NET_CAPABILITY_INTERNET)) {
-				Log.verbose("Network Capabilities changed, has Internet: " + false);
-				networks.remove(network.toString());
 			} else {
-				Log.verbose("Network Capabilities changed, has Internet: " + true);
-				networks.add(network.toString());
+				isNetworkAvailable = networkCapabilities.hasCapability(NetworkCapabilities.NET_CAPABILITY_INTERNET);
+				Log.verbose("Network Capabilities changed, has Internet: " + isNetworkAvailable);
 			}
 			checkNetworkConnectivity();
 		}
-	};
 
-	private Runnable retryRunnable = new Runnable() {
 		@Override
-		public void run() {
-			Log.verbose("Attempting to check for network connectivity again.");
-			retryCount++;
-			retryScheduled = false;
+		public void onBlockedStatusChanged(@NonNull Network network, boolean blocked)
+		{
+			isNetworkBlocked = blocked;
+			Log.verbose("Network Blocked Status changed, Blocked: " + isNetworkBlocked);
 			checkNetworkConnectivity();
 		}
 	};
 
+	private Runnable retryRunnable = () ->
+	{
+		Log.verbose("Attempting to check for network connectivity again.");
+		retryCount++;
+		retryScheduled = false;
+		checkNetworkConnectivity();
+	};
+
 	private void scheduleRetry() {
-		if (networkChangedListeners.size() == 0) {
+		if (networkChangedListeners.isEmpty()) {
 			Log.verbose("No listeners so not retrying. When a listener is added the connection will be checked.");
 			return;
 		}
@@ -183,7 +152,8 @@ public final class NetworkChangedManager implements NetworkConnectivityClient {
 
 	private void setNetworkState(ConnectivityState state) {
         NetworkTransportType updatedNetworkTransportType = calculateNetworkTransport(connectivityManager);
-        if (currentState == state && currentNetworkTransport == updatedNetworkTransportType) {
+		boolean updateIsNetworkMetered = connectivityManager.isActiveNetworkMetered();
+        if (currentState == state && currentNetworkTransport == updatedNetworkTransportType && isNetworkMetered == updateIsNetworkMetered) {
 			Log.verbose("Connectivity hasn't changed. Current state: " + currentState);
 			if (currentState != ConnectivityState.CONNECTION_AVAILABLE) {
 				scheduleRetry();
@@ -192,7 +162,8 @@ public final class NetworkChangedManager implements NetworkConnectivityClient {
 		}
 		currentState = state;
         currentNetworkTransport = updatedNetworkTransportType;
-		fireNetworkChangeListeners(state, currentNetworkTransport);
+		isNetworkMetered = updateIsNetworkMetered;
+		fireNetworkChangeListeners(state, currentNetworkTransport, isNetworkMetered);
 		Log.verbose("Network connectivity changed. New connectivity state: " + state);
 
 		if (currentState != ConnectivityState.CONNECTION_AVAILABLE) {
@@ -218,88 +189,74 @@ public final class NetworkChangedManager implements NetworkConnectivityClient {
 
 		networkCheckInProgress = true;
 		final ExecutorService executor = Executors.newSingleThreadExecutor();
-		final Runnable timeoutRunnable = new Runnable() {
-			@Override
-			public void run() {
-				Log.verbose("Unable to connect to: " + getCurrentHostResolutionAddress());
-				networkCheckInProgress = false;
-				executor.shutdownNow();
-				setNetworkState(ConnectivityState.NO_CONNECTION);
-				scheduleRetry();
-			}
+		final Runnable timeoutRunnable = () ->
+		{
+			Log.verbose("Unable to connect to: " + getCurrentHostResolutionAddress());
+			networkCheckInProgress = false;
+			executor.shutdownNow();
+			setNetworkState(ConnectivityState.NO_CONNECTION);
+			scheduleRetry();
 		};
 		internalScheduler.postDelayed(timeoutRunnable, HOSTNAME_RESOLUTION_TIMEOUT_MS * HOST_RESOLUTION_ADDRESSES.length + 100);
-		executor.execute(new Runnable() {
-			@Override
-			public void run() {
-				android.net.TrafficStats.setThreadStatsTag(1);
-				HttpURLConnection urlConnection = null;
-				boolean connectedSuccessfully = false;
-				// Attempt to connect to any of the hostnames. If any succeed we are connected to
-				// the internet.
-				for (int i = 0; i < HOST_RESOLUTION_ADDRESSES.length; i++) {
-					try {
-						Log.verbose("Verifying internet connection with host: " + getCurrentHostResolutionAddress());
-						URL url = new URL(getCurrentHostResolutionAddress());
-						urlConnection = (HttpURLConnection) url.openConnection();
-						urlConnection.setUseCaches(false);
-						urlConnection.setRequestMethod("HEAD");
-						urlConnection.setConnectTimeout((int) HOSTNAME_RESOLUTION_TIMEOUT_MS);
-						urlConnection.setReadTimeout((int) HOSTNAME_RESOLUTION_TIMEOUT_MS);
-						urlConnection.getInputStream().close();
-						connectedSuccessfully = true;
-						break;
-					} catch (MalformedURLException e) {
-						Log.error("Malformed URL, this should never happen. Please fix, url: " + getCurrentHostResolutionAddress());
-					} catch (IOException e) {
-						Log.verbose("Unable to connect to: " + getCurrentHostResolutionAddress());
-					} catch (Exception e) {
-						Log.verbose("Unable to connect to: " + getCurrentHostResolutionAddress() + ", exception: " + e.toString());
-					} finally {
-						if (urlConnection != null) {
-							urlConnection.disconnect();
-						}
-						if (!connectedSuccessfully) {
-							/*
-							 * Move to the next resolution address.
-							 * When we try again we should do it with the next available URL. For
-							 * some reason on some networks we've had issues with
-							 * https://example.com so using https://google.com and others will give
-							 * us more redundancy
-							 */
-							nextHostResolutionAddress();
-						}
+		executor.execute(() ->
+		{
+			android.net.TrafficStats.setThreadStatsTag(1);
+			HttpURLConnection urlConnection = null;
+			boolean connectedSuccessfully = false;
+			// Attempt to connect to any of the hostnames. If any succeed we are connected to
+			// the internet.
+			for (int i = 0; i < HOST_RESOLUTION_ADDRESSES.length; i++) {
+				try {
+					Log.verbose("Verifying internet connection with host: " + getCurrentHostResolutionAddress());
+					URL url = new URL(getCurrentHostResolutionAddress());
+					urlConnection = (HttpURLConnection) url.openConnection();
+					urlConnection.setUseCaches(false);
+					urlConnection.setRequestMethod("HEAD");
+					urlConnection.setConnectTimeout((int) HOSTNAME_RESOLUTION_TIMEOUT_MS);
+					urlConnection.setReadTimeout((int) HOSTNAME_RESOLUTION_TIMEOUT_MS);
+					urlConnection.getInputStream().close();
+					connectedSuccessfully = true;
+					break;
+				} catch (MalformedURLException e) {
+					Log.error("Malformed URL, this should never happen. Please fix, url: " + getCurrentHostResolutionAddress());
+				} catch (IOException e) {
+					Log.verbose("Unable to connect to: " + getCurrentHostResolutionAddress());
+				} catch (Exception e) {
+					Log.verbose("Unable to connect to: " + getCurrentHostResolutionAddress() + ", exception: " + e.toString());
+				} finally {
+					if (urlConnection != null) {
+						urlConnection.disconnect();
+					}
+					if (!connectedSuccessfully) {
+						/*
+						 * Move to the next resolution address.
+						 * When we try again we should do it with the next available URL. For
+						 * some reason on some networks we've had issues with
+						 * https://example.com so using https://google.com and others will give
+						 * us more redundancy
+						 */
+						nextHostResolutionAddress();
 					}
 				}
+			}
 
-				internalScheduler.removeCallbacks(timeoutRunnable);
-				networkCheckInProgress = false;
-				if (connectedSuccessfully) {
-					setNetworkState(ConnectivityState.CONNECTION_AVAILABLE);
-				} else {
-					setNetworkState(ConnectivityState.NO_CONNECTION);
-					scheduleRetry();
-				}
-				Log.verbose("Full network check complete. State: " + currentState);
-
-				executor.shutdownNow();
+			internalScheduler.removeCallbacks(timeoutRunnable);
+			networkCheckInProgress = false;
+			if (connectedSuccessfully) {
+				setNetworkState(ConnectivityState.CONNECTION_AVAILABLE);
+			} else {
+				setNetworkState(ConnectivityState.NO_CONNECTION);
+				scheduleRetry();
 			}
+			Log.verbose("Full network check complete. State: " + currentState);
+
+			executor.shutdownNow();
 		});
 	}
 
 	@NonNull
 	private ConnectivityState calculateNetworkConnectivityNaively() {
-		if (networks.isEmpty() || connectivityManager == null) {
-			return ConnectivityState.NO_CONNECTION;
-		} else {
-			NetworkInfo activeNetworkInfo = connectivityManager.getActiveNetworkInfo();
-
-			if (activeNetworkInfo != null && activeNetworkInfo.isAvailable() && activeNetworkInfo.isConnected()) {
-				return ConnectivityState.CONNECTION_AVAILABLE;
-			} else {
-				return ConnectivityState.NO_CONNECTION;
-			}
-		}
+		return isNetworkAvailable && !isNetworkBlocked ? ConnectivityState.CONNECTION_AVAILABLE : ConnectivityState.NO_CONNECTION;
 	}
 
 	private String getCurrentHostResolutionAddress() {
@@ -321,11 +278,6 @@ public final class NetworkChangedManager implements NetworkConnectivityClient {
 
 	@Override
 	public boolean addListener(Listener listener, boolean fireImmediately) {
-		if (Build.VERSION.SDK_INT < Build.VERSION_CODES.LOLLIPOP) {
-			// unsupported before Lollipop
-			return false;
-		}
-
 		for (WeakReference<Listener> currentListener : networkChangedListeners) {
 			if (currentListener.get() == listener) {
 				return false;
@@ -341,64 +293,69 @@ public final class NetworkChangedManager implements NetworkConnectivityClient {
 		}
 		// Current state will never be null after calling checkNetworkConnectivity
 		if (fireImmediately && currentState != null) {
-			fireNetworkChangeListenerInternal(listener, currentState, currentNetworkTransport);
+			fireNetworkChangeListenerInternal(listener, currentState, currentNetworkTransport, isNetworkMetered);
 		}
 		return true;
 	}
 
 	@Override
 	public boolean removeListener(Listener listener) {
-		if (Build.VERSION.SDK_INT < Build.VERSION_CODES.LOLLIPOP) {
-			// unsupported before Lollipop
+		if (!networkChangedListeners.removeIf(reference ->
+			{
+				Listener item = reference.get();
+				return item == listener || item == null;
+			}))
+		{
 			return false;
 		}
 
-		Iterator<WeakReference<Listener>> changeListenersIterator = networkChangedListeners.iterator();
-		while (changeListenersIterator.hasNext()) {
-			WeakReference<Listener> currentListener = changeListenersIterator.next();
-			if (currentListener.get() == listener) {
-				changeListenersIterator.remove();
-				return true;
-			}
-		}
-		/*
-		 * If there are no listeners, clear any retries. When a listener is added we'll check the
-		 * connection again.
-		 */
-		if (networkChangedListeners.size() == 0) {
+		if (networkChangedListeners.isEmpty())
+		{
 			clearRetry();
 		}
-		return false;
+
+		return true;
 	}
 
-	private void fireNetworkChangeListeners(ConnectivityState state, NetworkTransportType networkTransportType) {
+	private void fireNetworkChangeListeners(ConnectivityState state, NetworkTransportType networkTransportType, boolean isNetworkMetered) {
 		// can be called from callbacks that should return quickly
 		final ExecutorService executor = Executors.newSingleThreadExecutor();
-		executor.execute(new Runnable() {
-			@Override
-			public void run() {
-				Iterator<WeakReference<Listener>> changeListenersIterator = networkChangedListeners.iterator();
-				while (changeListenersIterator.hasNext()) {
-					WeakReference<Listener> listenerWeakReference = changeListenersIterator.next();
-					Listener listener = listenerWeakReference.get();
-					if (listener == null) {
-						changeListenersIterator.remove();
-					} else {
-						fireNetworkChangeListenerInternal(listener, state, networkTransportType);
-					}
+		executor.execute(() ->
+		{
+			boolean needsCleanup = false;
+
+			for (WeakReference<Listener> reference : networkChangedListeners)
+			{
+				Listener listener = reference.get();
+				if (listener == null)
+				{
+					needsCleanup = true;
+				}
+				else
+				{
+					fireNetworkChangeListenerInternal(listener, state, networkTransportType, isNetworkMetered);
 				}
-				executor.shutdownNow();
 			}
+
+			if (needsCleanup)
+			{
+				if (networkChangedListeners.removeIf(reference -> reference.get() == null) && networkChangedListeners.isEmpty())
+				{
+					clearRetry();
+				}
+			}
+
+			executor.shutdownNow();
 		});
 	}
 
-	private void fireNetworkChangeListenerInternal(Listener listener, ConnectivityState state, NetworkTransportType networkTransportType) {
+	private void fireNetworkChangeListenerInternal(Listener listener, ConnectivityState state, NetworkTransportType networkTransportType, boolean isNetworkMetered) {
 		switch (state) {
 			case NO_CONNECTION:
 				listener.onNetworkLost();
 				break;
 			case CONNECTION_AVAILABLE:
-				listener.onNetworkAvailable(networkTransportType);
+				listener.onNetworkAvailable(networkTransportType, isNetworkMetered);
 				break;
 		}
 	}
@@ -408,56 +365,34 @@ public final class NetworkChangedManager implements NetworkConnectivityClient {
 	{
 		return calculateNetworkTransport(connectivityManager);
 	}
-	
+
 	@Override
 	public void checkConnectivity() {
 		checkNetworkConnectivity();
 	}
 
     private NetworkTransportType calculateNetworkTransport(@NonNull ConnectivityManager connectivityManager) {
-       if (Build.VERSION.SDK_INT < Build.VERSION_CODES.M) {
-           NetworkInfo networkInfo = connectivityManager.getActiveNetworkInfo();
-           if (networkInfo == null) {
-               return NetworkTransportType.UNKNOWN;
-           }
-
-           switch (networkInfo.getType()) {
-               case ConnectivityManager.TYPE_WIFI:
-                   return NetworkTransportType.WIFI;
-               case ConnectivityManager.TYPE_MOBILE:
-                   return NetworkTransportType.CELLULAR;
-               case ConnectivityManager.TYPE_ETHERNET:
-                   return NetworkTransportType.ETHERNET;
-               case ConnectivityManager.TYPE_BLUETOOTH:
-                   return NetworkTransportType.BLUETOOTH;
-               case ConnectivityManager.TYPE_VPN:
-                   return NetworkTransportType.VPN;
-               default:
-                   return NetworkTransportType.UNKNOWN;
-           }
-       } else {
-           Network activeNetwork = connectivityManager.getActiveNetwork();
-           if (activeNetwork == null) {
-               return NetworkTransportType.UNKNOWN;
-           }
-           NetworkCapabilities capabilities = connectivityManager.getNetworkCapabilities(activeNetwork);
-           if (capabilities == null) {
-               return NetworkTransportType.UNKNOWN;
-           }
-
-           if (capabilities.hasTransport(NetworkCapabilities.TRANSPORT_WIFI)) {
-               return NetworkTransportType.WIFI;
-           } else if (capabilities.hasTransport(NetworkCapabilities.TRANSPORT_CELLULAR)) {
-               return NetworkTransportType.CELLULAR;
-           } else if (capabilities.hasTransport(NetworkCapabilities.TRANSPORT_ETHERNET)) {
-               return NetworkTransportType.ETHERNET;
-           } else if (capabilities.hasTransport(NetworkCapabilities.TRANSPORT_BLUETOOTH)) {
-               return NetworkTransportType.BLUETOOTH;
-           } else if (capabilities.hasTransport(NetworkCapabilities.TRANSPORT_VPN)) {
-               return NetworkTransportType.VPN;
-           } else {
-               return NetworkTransportType.UNKNOWN;
-           }
-       }
-   }
+		Network activeNetwork = connectivityManager.getActiveNetwork();
+		if (activeNetwork == null) {
+			return NetworkTransportType.UNKNOWN;
+		}
+		NetworkCapabilities capabilities = connectivityManager.getNetworkCapabilities(activeNetwork);
+		if (capabilities == null) {
+			return NetworkTransportType.UNKNOWN;
+		}
+
+		if (capabilities.hasTransport(NetworkCapabilities.TRANSPORT_WIFI)) {
+			return NetworkTransportType.WIFI;
+		} else if (capabilities.hasTransport(NetworkCapabilities.TRANSPORT_CELLULAR)) {
+			return NetworkTransportType.CELLULAR;
+		} else if (capabilities.hasTransport(NetworkCapabilities.TRANSPORT_ETHERNET)) {
+			return NetworkTransportType.ETHERNET;
+		} else if (capabilities.hasTransport(NetworkCapabilities.TRANSPORT_BLUETOOTH)) {
+			return NetworkTransportType.BLUETOOTH;
+		} else if (capabilities.hasTransport(NetworkCapabilities.TRANSPORT_VPN)) {
+			return NetworkTransportType.VPN;
+		} else {
+			return NetworkTransportType.UNKNOWN;
+		}
+	}
 }
diff --git a/Engine/Build/Android/Java/src/com/epicgames/unreal/network/NetworkConnectivityClient.java b/Engine/Build/Android/Java/src/com/epicgames/unreal/network/NetworkConnectivityClient.java
--- a/Engine/Build/Android/Java/src/com/epicgames/unreal/network/NetworkConnectivityClient.java
+++ b/Engine/Build/Android/Java/src/com/epicgames/unreal/network/NetworkConnectivityClient.java
@@ -3,20 +3,21 @@
 package com.epicgames.unreal.network;
 
 import android.content.Context;
+
 import androidx.annotation.NonNull;
 
 public interface NetworkConnectivityClient {
-    public enum NetworkTransportType {
-        WIFI,
-        VPN,
-        ETHERNET,
-        CELLULAR,
-        BLUETOOTH,
-        UNKNOWN
-    }
-
-    public interface Listener {
-        void onNetworkAvailable(NetworkTransportType networkTransportType);
+	enum NetworkTransportType {
+		WIFI,
+		VPN,
+		ETHERNET,
+		CELLULAR,
+		BLUETOOTH,
+		UNKNOWN
+	}
+
+	interface Listener {
+		void onNetworkAvailable(NetworkTransportType networkTransportType, boolean isNetworkMetered);
 
 		void onNetworkLost();
 	}
diff --git a/Engine/Plugins/Online/Android/AndroidFetchBackgroundDownload/Source/AndroidFetchBackgroundDownload/AndroidFetchBackgroundDownload_UPL.xml b/Engine/Plugins/Online/Android/AndroidFetchBackgroundDownload/Source/AndroidFetchBackgroundDownload/AndroidFetchBackgroundDownload_UPL.xml
--- a/Engine/Plugins/Online/Android/AndroidFetchBackgroundDownload/Source/AndroidFetchBackgroundDownload/AndroidFetchBackgroundDownload_UPL.xml
+++ b/Engine/Plugins/Online/Android/AndroidFetchBackgroundDownload/Source/AndroidFetchBackgroundDownload/AndroidFetchBackgroundDownload_UPL.xml
@@ -45,6 +45,12 @@
 		</insert>
 	</buildGradleAdditions>
 
+	<androidManifestUpdates>
+		<addElements tag="application">
+			<receiver android:name="com.epicgames.unreal.download.DownloadManager$CancelReceiver" android:exported="false"/>
+		</addElements>
+	</androidManifestUpdates>
+	
 	<!-- optional additions to proguard -->
 	<proguardAdditions>
 		<insert>
@@ -52,6 +58,10 @@
 
 			# Android Background HTTP Classes with native calls ...
 
+			-keep class com.epicgames.unreal.download.DownloadManager {
+			public *;
+			}
+
 			-keep class com.epicgames.unreal.download.UEDownloadWorker {
 			public *;
 			}
@@ -73,5 +83,12 @@
 	<!-- set build version to BranchName-Changelist, see C++ equivalent in AndroidPlatformBackgroundHttpRequest.cpp-->
 	<gameActivityClassAdditions>
 		<insertValue value="public static final String BGDL_BUILD_VERSION = &quot;$S(EngineBranchName)&quot;;"/>
+		<insertNewLine/>
 	</gameActivityClassAdditions>
+
+	<gameActivityOnResumeAdditions>
+		<insert>
+			com.epicgames.unreal.download.DownloadManager.Resume();
+		</insert>
+	</gameActivityOnResumeAdditions>
 </root>
diff --git a/Engine/Plugins/Online/Android/AndroidFetchBackgroundDownload/Source/AndroidFetchBackgroundDownload/Java/download/DownloadManager.java b/Engine/Plugins/Online/Android/AndroidFetchBackgroundDownload/Source/AndroidFetchBackgroundDownload/Java/download/DownloadManager.java
new file mode 100644
--- /dev/null
+++ b/Engine/Plugins/Online/Android/AndroidFetchBackgroundDownload/Source/AndroidFetchBackgroundDownload/Java/download/DownloadManager.java
@@ -0,0 +1,1230 @@
+// Copyright Epic Games, Inc. All Rights Reserved.
+package com.epicgames.unreal.download;
+
+import android.annotation.SuppressLint;
+import android.app.ForegroundServiceStartNotAllowedException;
+import android.app.Notification;
+import android.app.NotificationChannel;
+import android.app.NotificationManager;
+import android.app.PendingIntent;
+import android.content.BroadcastReceiver;
+import android.content.Context;
+import android.content.Intent;
+import android.content.IntentFilter;
+import android.content.SharedPreferences;
+import android.net.ConnectivityManager;
+import android.os.Build;
+import android.os.Handler;
+import android.os.Looper;
+import android.provider.Settings;
+
+import androidx.annotation.Keep;
+import androidx.annotation.MainThread;
+import androidx.annotation.NonNull;
+import androidx.concurrent.futures.CallbackToFutureAdapter;
+import androidx.core.app.NotificationCompat;
+import androidx.core.net.ConnectivityManagerCompat;
+import androidx.work.BackoffPolicy;
+import androidx.work.ExistingWorkPolicy;
+import androidx.work.ForegroundInfo;
+import androidx.work.ListenableWorker;
+import androidx.work.OneTimeWorkRequest;
+import androidx.work.Operation;
+import androidx.work.WorkManager;
+import androidx.work.WorkerParameters;
+
+import com.epicgames.unreal.CellularReceiver;
+import com.epicgames.unreal.GameActivity;
+import com.epicgames.unreal.GameApplication;
+import com.epicgames.unreal.Logger;
+import com.epicgames.unreal.download.datastructs.DownloadDescription;
+import com.epicgames.unreal.download.datastructs.DownloadNotificationDescription;
+import com.epicgames.unreal.download.datastructs.DownloadQueueDescription;
+import com.epicgames.unreal.download.datastructs.DownloadWorkerParameterKeys;
+import com.epicgames.unreal.download.fetch.FetchManagerInterfaceHelper;
+import com.epicgames.unreal.download.fetch.FetchManagerNew;
+import com.epicgames.unreal.network.NetworkChangedManager;
+import com.epicgames.unreal.network.NetworkConnectivityClient;
+import com.google.common.util.concurrent.ListenableFuture;
+
+import java.io.IOException;
+import java.nio.charset.StandardCharsets;
+import java.nio.file.FileVisitResult;
+import java.nio.file.Files;
+import java.nio.file.Path;
+import java.nio.file.Paths;
+import java.nio.file.SimpleFileVisitor;
+import java.nio.file.attribute.BasicFileAttributes;
+import java.text.NumberFormat;
+import java.util.ArrayList;
+import java.util.Comparator;
+import java.util.HashMap;
+import java.util.List;
+import java.util.Map;
+import java.util.Objects;
+import java.util.concurrent.ConcurrentHashMap;
+import java.util.concurrent.ExecutionException;
+import java.util.concurrent.TimeUnit;
+import java.util.concurrent.atomic.AtomicInteger;
+import java.util.stream.Collectors;
+import java.util.stream.Stream;
+
+import dalvik.annotation.optimization.FastNative;
+
+@Keep
+public class DownloadManager
+{
+	static final @NonNull Logger Log = new Logger("UE", DownloadManager.class.getSimpleName());
+	
+	static final @NonNull Handler Handler = new Handler(Looper.getMainLooper());
+
+	final static String WorkName = "BackgroundHttpDownload";
+	final static int InitialBackoffDelayInSeconds = 10;
+
+	static SharedPreferences BackgroundPreferences;
+	
+	@SuppressLint("StaticFieldLeak")
+	static Task Task;
+	
+	@Keep
+	@FastNative
+	native static void OnProgress(@NonNull String DestinationLocation, long BytesWrittenSinceLastCall, long TotalBytesWritten);
+	@Keep
+	@FastNative
+	native static void OnComplete(@NonNull String DestinationLocation, boolean bDidSucceed, long TotalBytesDownloaded, long DownloadDuration, long DownloadStartTimeUTC, long DownloadEndTimeUTC);
+	
+	static boolean IsInForeground()
+	{
+		return BackgroundPreferences != null;
+	}
+
+	static @NonNull SharedPreferences GetBackgroundPreferences(@NonNull Context context)
+	{
+		return context.getSharedPreferences("BackgroundPreferences", Context.MODE_PRIVATE);
+	}
+	
+	static @NonNull String GetTempDownloadRootPath(@NonNull SharedPreferences backgroundPreferences)
+	{
+		String tempDownloadRootPath = backgroundPreferences.getString("TempDownloadRootPath", "");
+		assert !tempDownloadRootPath.isEmpty();
+		return tempDownloadRootPath;
+	}
+	
+	static ListenableFuture<Operation.State.SUCCESS> EnqueueWorker(@NonNull Context context)
+	{
+		assert IsInForeground();
+		
+		return WorkManager.getInstance(context).enqueueUniqueWork(WorkName, ExistingWorkPolicy.KEEP, new OneTimeWorkRequest.Builder(Worker.class)
+			.setBackoffCriteria(BackoffPolicy.EXPONENTIAL, InitialBackoffDelayInSeconds, TimeUnit.SECONDS)
+			.build()).getResult();
+	}
+
+	static void CancelWorker(@NonNull Context context)
+	{
+		WorkManager.getInstance(context).cancelUniqueWork(WorkName);
+	}
+	
+	@Keep
+	public static class CancelReceiver extends BroadcastReceiver
+	{
+		@MainThread
+		@Override
+		public void onReceive(@NonNull Context context, Intent intent)
+		{
+			if (intent != null)
+			{
+				if (Task != null && Task.Worker != null)
+				{
+					if (!IsInForeground() || GameActivity.IsActivityPaused())
+					{
+						Task.Worker.Cancel();
+					}
+				}
+			}
+		}
+	}
+
+	public static class RestrictBackgroundStatusReceiver extends BroadcastReceiver
+	{
+		@NonNull ConnectivityManager ConnectivityManager;
+		
+		public RestrictBackgroundStatusReceiver(@NonNull Context context)
+		{
+			assert Task != null;
+			
+			ConnectivityManager = context.getSystemService(ConnectivityManager.class);
+			
+			Task.RestrictBackgroundStatus = ConnectivityManager.getRestrictBackgroundStatus();
+		}
+		
+		@MainThread
+		@Override
+		public void onReceive(@NonNull Context context, Intent intent)
+		{
+			if (intent != null)
+			{
+				if (Task != null)
+				{
+					Task.SetRestrictBackgroundStatus(ConnectivityManager.getRestrictBackgroundStatus());
+				}
+			}
+		}
+	}
+
+	static class Task implements Runnable, DownloadProgressListener, NetworkConnectivityClient.Listener, SharedPreferences.OnSharedPreferenceChangeListener
+	{
+		final @NonNull Context Context;
+		final @NonNull Path DownloadDescriptionsPath;
+
+		final @NonNull ConcurrentHashMap<String, DownloadDescription> DownloadDescriptions = new ConcurrentHashMap<>();
+		final @NonNull AtomicInteger RequestCount = new AtomicInteger();
+
+		FetchManagerNew FetchManager;
+		SharedPreferences CellularNetworkPreferences;
+		RestrictBackgroundStatusReceiver RestrictBackgroundStatusReceiver;
+		
+		boolean bIsPaused;
+		boolean bIsNetworkLost;
+		NetworkConnectivityClient.NetworkTransportType NetworkTransportType = NetworkConnectivityClient.NetworkTransportType.UNKNOWN;
+		boolean bIsNetworkMetered;
+		@ConnectivityManagerCompat.RestrictBackgroundStatus int RestrictBackgroundStatus;
+
+		ListenableFuture<Operation.State.SUCCESS> EnqueueWorkerFuture;
+		Worker Worker;
+
+		public Task(@NonNull Context context, @NonNull SharedPreferences backgroundPreferences)
+		{
+			Context = context;
+			DownloadDescriptionsPath = Paths.get(GetTempDownloadRootPath(backgroundPreferences), "DownloadDescriptions");
+		}
+		
+		@MainThread
+		void Initialize()
+		{
+			assert Task == null;
+			Task = this;
+
+			FetchManager = (FetchManagerNew)FetchManagerInterfaceHelper.GetSharedManager();
+			CellularNetworkPreferences = Context.getSharedPreferences("CellularNetworkPreferences", android.content.Context.MODE_PRIVATE);
+			RestrictBackgroundStatusReceiver = new RestrictBackgroundStatusReceiver(Context);
+			
+			NetworkChangedManager.getInstance().addListener(this, true);
+
+			CellularNetworkPreferences.registerOnSharedPreferenceChangeListener(this);
+
+			Context.registerReceiver(RestrictBackgroundStatusReceiver, new IntentFilter(ConnectivityManager.ACTION_RESTRICT_BACKGROUND_CHANGED));
+		}
+
+		@MainThread
+		public void Start()
+		{
+			Log.debug("Starting");
+			
+			assert Task == null;
+			assert IsInForeground();
+			
+			try
+			{
+				Files.createDirectories(DownloadDescriptionsPath);
+			}
+			catch (IOException e)
+			{
+				e.printStackTrace();
+			}
+			
+			Initialize();
+		}
+		
+		@MainThread
+		public void Restart(@NonNull SharedPreferences backgroundPreferences)
+		{
+			Log.debug("Restarting");
+			
+			assert Task == null;
+			assert !IsInForeground();
+
+			ArrayList<DownloadDescription> downloadDescriptions;
+			try (Stream<Path> files = Files.list(DownloadDescriptionsPath))
+			{
+				downloadDescriptions = files.map(file ->
+				{
+					try
+					{
+						return DownloadDescription.FromJSON(new String(Files.readAllBytes(file), StandardCharsets.UTF_8));
+					}
+					catch (IOException ignored)
+					{
+						return null;
+					}
+				}).filter(Objects::nonNull).collect(Collectors.toCollection(ArrayList::new));
+			}
+			catch (IOException ignored)
+			{
+				return;
+			}
+
+			if (downloadDescriptions.isEmpty())
+			{
+				return;
+			}
+
+			Initialize();
+
+			DownloadQueueDescription downloadQueueDescription = new DownloadQueueDescription(downloadDescriptions, backgroundPreferences);
+			downloadQueueDescription.ProgressListener = this;
+			FetchManager.Enqueue(Context, downloadQueueDescription);
+		}
+
+		@MainThread
+		public boolean IsWaitingForCellularApproval()
+		{
+			return NetworkTransportType == NetworkConnectivityClient.NetworkTransportType.CELLULAR &&
+				!(RestrictBackgroundStatus == ConnectivityManagerCompat.RESTRICT_BACKGROUND_STATUS_WHITELISTED || CellularNetworkPreferences.getInt("AllowCellular", 0) > 0);
+		}
+
+		@MainThread
+		public boolean IsResumable()
+		{
+			return !bIsNetworkLost && !bIsPaused && (IsInForeground() || !IsWaitingForCellularApproval());
+		}
+		
+		@MainThread
+		public void SetRestrictBackgroundStatus(@ConnectivityManagerCompat.RestrictBackgroundStatus int restrictBackgroundStatus)
+		{
+			if (RestrictBackgroundStatus != restrictBackgroundStatus)
+			{
+				RestrictBackgroundStatus = restrictBackgroundStatus;
+				
+				if (!IsInForeground())
+				{
+					if (Task.IsResumable())
+					{
+						Log.debug("SetRestrictBackgroundStatus, resuming");
+						FetchManager.ResumeAllDownloads();
+					}
+					else
+					{
+						Log.debug("SetRestrictBackgroundStatus, pausing");
+						FetchManager.PauseAllDownloads();
+					}
+				}
+			}
+		}
+
+		@MainThread
+		public void Enqueue(@NonNull DownloadDescription downloadDescription, @NonNull SharedPreferences backgroundPreferences)
+		{
+			assert IsInForeground();
+
+			DownloadQueueDescription downloadQueueDescription = new DownloadQueueDescription(new ArrayList<>(List.of(downloadDescription)), backgroundPreferences);
+			downloadQueueDescription.ProgressListener = this;
+			FetchManager.Enqueue(Context, downloadQueueDescription);
+
+			if (Worker == null && (EnqueueWorkerFuture == null || EnqueueWorkerFuture.isDone()))
+			{
+				EnqueueWorkerFuture = EnqueueWorker(Context);
+			}
+		}
+
+		@MainThread
+		public void Pause()
+		{
+			assert Task == this;
+			assert Worker == null;
+
+			if (GameActivity.IsActivityPaused())
+			{
+				if (!bIsPaused)
+				{
+					if (IsResumable())
+					{
+						Log.debug("Pause, pausing");
+						FetchManager.PauseAllDownloads();
+					}
+
+					bIsPaused = true;
+				}
+			}
+			else
+			{
+				EnqueueWorkerFuture = EnqueueWorker(Context);
+			}
+		}
+
+		@MainThread
+		public void Resume()
+		{
+			assert Task == this;
+			assert Worker == null;
+
+			if (bIsPaused)
+			{
+				bIsPaused = false;
+
+				if (IsResumable())
+				{
+					Log.debug("Resume, resuming");
+					FetchManager.ResumeAllDownloads();
+				}
+
+				if (EnqueueWorkerFuture == null || EnqueueWorkerFuture.isDone())
+				{
+					EnqueueWorkerFuture = EnqueueWorker(Task.Context);
+				}
+			}
+		}
+
+		@MainThread
+		void Dispose()
+		{
+			assert Task == this;
+			Task = null;
+			
+			Context.unregisterReceiver(RestrictBackgroundStatusReceiver);
+
+			CellularNetworkPreferences.unregisterOnSharedPreferenceChangeListener(this);
+
+			NetworkChangedManager.getInstance().removeListener(this);
+		}
+
+		@MainThread
+		void Stop()
+		{
+			Log.debug("Stopping");
+			
+			Dispose();
+
+			FetchManager.StopWork(null);
+			
+			Handler.removeCallbacks(this);
+		}
+
+		@MainThread
+		void AllComplete()
+		{
+			Dispose();
+			
+			final boolean didAllSucceed = DownloadDescriptions.isEmpty();
+			
+			if (didAllSucceed)
+			{
+				Context.getSharedPreferences("UEDownloadWorker", android.content.Context.MODE_PRIVATE).edit().putString("AllDownloadsCompletedTime", NumberFormat.getInstance().format(System.currentTimeMillis() / 1000.0f)).commit();
+				
+				if (!IsInForeground())
+				{
+					CellularNetworkPreferences.edit().putInt("AllowCellular", 0).commit();
+				}
+			}
+
+			if (Worker != null)
+			{
+				Worker.AllComplete(didAllSucceed ? ListenableWorker.Result.success() : ListenableWorker.Result.retry());
+			}
+			else
+			{
+				CancelWorker(Context);
+			}
+		}
+
+		@MainThread
+		public void CancelAll()
+		{
+			assert Task == this;
+
+			Stop();
+
+			for (DownloadDescription downloadDescription : DownloadDescriptions.values())
+			{
+				try
+				{
+					Files.delete(DownloadDescriptionsPath.resolve(Paths.get(downloadDescription.DestinationLocation).getFileName()));
+				}
+				catch (IOException ignored)
+				{
+				}
+			}
+
+			if (Worker != null)
+			{
+				Worker.AllComplete(ListenableWorker.Result.success());
+			}
+			else
+			{
+				CancelWorker(Context);
+			}
+		}
+		
+		boolean Removed(@NonNull DownloadDescription downloadDescription)
+		{
+			try
+			{
+				Files.delete(DownloadDescriptionsPath.resolve(Paths.get(downloadDescription.DestinationLocation).getFileName()));
+			}
+			catch (IOException ignored)
+			{
+			}
+
+			int requestCount = RequestCount.decrementAndGet();
+			assert requestCount >= 0;
+			return requestCount == 0;
+		}
+
+		@MainThread
+		public void Pause(@NonNull String destinationLocation)
+		{
+			assert Task == this;
+
+			Log.debug("@*@ PauseRequest: " + destinationLocation);
+			FetchManager.PauseDownload(destinationLocation, true);
+		}
+
+		@MainThread
+		public void Resume(@NonNull String destinationLocation)
+		{
+			assert Task == this;
+
+			Log.debug("@*@ ResumeRequest: " + destinationLocation);
+			FetchManager.PauseDownload(destinationLocation, false);
+		}
+		
+		@MainThread
+		public void Cancel(@NonNull String destinationLocation)
+		{
+			assert Task == this;
+
+			FetchManager.CancelDownload(destinationLocation);
+
+			DownloadDescription downloadDescription = DownloadDescriptions.remove(destinationLocation);
+			if (downloadDescription != null && Removed(downloadDescription))
+			{
+				AllComplete();
+			}
+		}
+
+		@MainThread
+		@Override
+		public void run()
+		{
+			if (Task != null)
+			{
+				assert Task == this;
+
+				if (RequestCount.get() == 0)
+				{
+					AllComplete();
+				}
+			}
+		}
+
+		@Override
+		public void OnDownloadProgress(@NonNull String DestinationLocation, long TotalBytesWritten, long BytesWrittenSinceLastCall)
+		{
+			if (IsInForeground())
+			{
+				OnProgress(DestinationLocation, TotalBytesWritten, BytesWrittenSinceLastCall);
+			}
+		}
+
+		@MainThread
+		@Override
+		public void OnDownloadGroupProgress(int GroupID, final int Progress, final boolean Indeterminate)
+		{
+			if (Worker != null)
+			{
+				Worker.UpdateProgress(Indeterminate ? -1 : Progress);
+			}
+		}
+
+		@Override
+		public void OnDownloadComplete(@NonNull String DestinationLocation, UEDownloadWorker.EDownloadCompleteReason CompleteReason)
+		{
+			assert Thread.holdsLock(FetchManager.SyncRoot());
+			
+			final boolean bDidSucceed = CompleteReason == UEDownloadWorker.EDownloadCompleteReason.Success;
+			if (IsInForeground() || bDidSucceed)
+			{
+				DownloadDescription downloadDescription = DownloadDescriptions.remove(DestinationLocation);
+				if (downloadDescription != null)
+				{
+					if (Removed(downloadDescription))
+					{
+						Handler.post(this);
+					}
+
+					if (IsInForeground())
+					{
+						OnComplete(DestinationLocation, bDidSucceed, downloadDescription.TotalDownloadedBytes, downloadDescription.DownloadDuration, downloadDescription.DownloadStartTimeUTC, downloadDescription.DownloadEndTimeUTC);
+					}
+				}
+			}
+		}
+
+		@Override
+		public void OnDownloadMetrics(@NonNull String DestinationLocation, long TotalBytesDownloaded, long DownloadDuration, long DownloadStartTimeUTC, long DownloadEndTimeUTC)
+		{
+			DownloadDescription downloadDescription = DownloadDescriptions.get(DestinationLocation);
+			if (downloadDescription != null)
+			{
+				downloadDescription.TotalDownloadedBytes = TotalBytesDownloaded;
+				downloadDescription.DownloadDuration = DownloadDuration;
+				downloadDescription.DownloadStartTimeUTC = DownloadStartTimeUTC;
+				downloadDescription.DownloadEndTimeUTC = DownloadEndTimeUTC;
+			}
+		}
+
+		@Override
+		public void OnAllDownloadsComplete(boolean bDidAllRequestsSucceed)
+		{
+		}
+
+		@MainThread
+		@Override
+		public void OnDownloadEnqueued(@NonNull DownloadDescription DownloadDescription)
+		{
+			assert Thread.holdsLock(FetchManager.SyncRoot());
+			
+			if (DownloadDescriptions.put(DownloadDescription.DestinationLocation, DownloadDescription) == null)
+			{
+				RequestCount.incrementAndGet();
+			}
+				
+			if (IsInForeground())
+			{
+				try
+				{
+					Files.write(DownloadDescriptionsPath.resolve(Paths.get(DownloadDescription.DestinationLocation).getFileName()), DownloadDescription.ToJSON().getBytes(StandardCharsets.UTF_8));
+				}
+				catch (IOException ignored)
+				{
+				}
+			}
+		}
+		
+		@Override
+		public void onNetworkAvailable(NetworkConnectivityClient.NetworkTransportType networkTransportType, boolean isNetworkMetered)
+		{
+			Handler.post(() ->
+			{
+				if (Task != null)
+				{
+					assert Task == this;
+
+					NetworkTransportType = networkTransportType;
+					bIsNetworkMetered = isNetworkMetered;
+					
+					if (bIsNetworkLost)
+					{
+						bIsNetworkLost = false;
+
+						if (IsResumable())
+						{
+							Log.debug("onNetworkAvailable, resuming");
+							FetchManager.ResumeAllDownloads();
+						}
+					}
+				}
+			});
+		}
+
+		@Override
+		public void onNetworkLost()
+		{
+			Handler.post(() ->
+			{
+				if (Task != null)
+				{
+					assert Task == this;
+					
+					if (!bIsNetworkLost)
+					{
+						if (IsResumable())
+						{
+							Log.debug("onNetworkLost, pausing");
+							FetchManager.PauseAllDownloads();
+						}
+
+						bIsNetworkLost = true;
+					}
+				}
+			});
+		}
+
+		@MainThread
+		@Override
+		public void onSharedPreferenceChanged(@NonNull SharedPreferences sharedPreferences, String key)
+		{
+			if (IsResumable())
+			{
+				Log.debug("onSharedPreferenceChanged, resuming");
+				FetchManager.ResumeAllDownloads();
+			}
+			else
+			{
+				Log.debug("onSharedPreferenceChanged, pausing");
+				FetchManager.PauseAllDownloads();
+			}
+		}
+	}
+
+	@Keep
+	public static class Worker extends ListenableWorker implements CallbackToFutureAdapter.Resolver<ListenableWorker.Result>, Runnable
+	{
+		CallbackToFutureAdapter.Completer<ListenableWorker.Result> Completer;
+
+		DownloadNotificationDescription NotificationDescription;
+		
+		PendingIntent PendingCancelIntent;
+		PendingIntent PendingApproveIntent;
+		PendingIntent PendingNotificationIntent;
+		
+		public Worker(@NonNull Context context, @NonNull WorkerParameters params)
+		{
+			super(context, params);
+		}
+		
+		void Initialize(@NonNull Context context, @NonNull SharedPreferences backgroundPreferences)
+		{
+			NotificationDescription = new DownloadNotificationDescription(backgroundPreferences, context, Log);
+
+			PendingCancelIntent = PendingIntent.getBroadcast(context, 0, new Intent(context, CancelReceiver.class), PendingIntent.FLAG_IMMUTABLE);
+			PendingApproveIntent = PendingIntent.getBroadcast(context, 0, new Intent(context, CellularReceiver.class), PendingIntent.FLAG_IMMUTABLE);
+			PendingNotificationIntent = PendingIntent.getActivity(context, NotificationDescription.NotificationID, new Intent(context, GameActivity.class)
+				.setFlags(Intent.FLAG_ACTIVITY_CLEAR_TOP | Intent.FLAG_ACTIVITY_SINGLE_TOP)
+				.putExtra("localNotificationID", NotificationDescription.NotificationID)
+				.putExtra("localNotificationAppLaunched", true), PendingIntent.FLAG_IMMUTABLE);
+			
+			NotificationManager notificationManager = context.getSystemService(NotificationManager.class);
+			if (notificationManager.getNotificationChannel(NotificationDescription.NotificationChannelID) == null)
+			{
+				String channelName = NotificationDescription.NotificationChannelName;
+				if (channelName.equals("ue-downloadworker-channel"))
+				{
+					// check for an override for default channel name
+					channelName = com.epicgames.unreal.LocalNotificationReceiver.getLocalizedResource(context, "UEBackgroundChannel", channelName);
+				}
+				notificationManager.createNotificationChannel(new NotificationChannel(NotificationDescription.NotificationChannelID, channelName, NotificationDescription.NotificationChannelImportance));
+			}
+		}
+
+		@NonNull Notification CreateAirplaneModeNotification(@NonNull Context context)
+		{
+			final boolean indeterminate = !NotificationDescription.ShouldShowPercentage || NotificationDescription.Indeterminate;
+			String notificationText;
+			if (NotificationDescription.CurrentProgress < NotificationDescription.MAX_PROGRESS)
+			{
+				notificationText = "";
+			}
+			else
+			{
+				notificationText = NotificationDescription.ContentCompleteText;
+			}
+
+			NotificationCompat.Builder builder = new NotificationCompat.Builder(context, NotificationDescription.NotificationChannelID)
+				.setContentTitle(NotificationDescription.AirplaneModeText)
+				.setTicker(NotificationDescription.TitleText)
+				.setContentText(notificationText)
+				.setContentIntent(PendingNotificationIntent)
+				.setProgress(NotificationDescription.MAX_PROGRESS, NotificationDescription.CurrentProgress, indeterminate)
+				.setOngoing(true)
+				.setOnlyAlertOnce(true)
+				.setSmallIcon(NotificationDescription.SmallIconResourceID)
+				.setSilent(true);
+			if (!IsInForeground())
+			{
+				builder.addAction(NotificationDescription.CancelIconResourceID, NotificationDescription.CancelText, PendingCancelIntent);
+			}
+			return builder.build();
+		}
+
+		@NonNull Notification CreateDataSaverEnabledNotification(@NonNull Context context)
+		{
+			final boolean indeterminate = !NotificationDescription.ShouldShowPercentage || NotificationDescription.Indeterminate;
+			String notificationText;
+			if (NotificationDescription.CurrentProgress < NotificationDescription.MAX_PROGRESS)
+			{
+				notificationText = "";
+			}
+			else
+			{
+				notificationText = NotificationDescription.ContentCompleteText;
+			}
+
+			NotificationCompat.Builder builder = new NotificationCompat.Builder(context, NotificationDescription.NotificationChannelID)
+				.setContentTitle(NotificationDescription.DataSaverEnabledText)
+				.setTicker(NotificationDescription.TitleText)
+				.setContentText(notificationText)
+				.setContentIntent(PendingNotificationIntent)
+				.setProgress(NotificationDescription.MAX_PROGRESS, NotificationDescription.CurrentProgress, indeterminate)
+				.setOngoing(true)
+				.setOnlyAlertOnce(true)
+				.setSmallIcon(NotificationDescription.SmallIconResourceID)
+				.setSilent(true);
+			if (!IsInForeground())
+			{
+				builder.addAction(NotificationDescription.CancelIconResourceID, NotificationDescription.CancelText, PendingCancelIntent);
+			}
+			return builder.build();
+		}
+
+		@NonNull Notification CreateNoInternetDownloadNotification(@NonNull Context context)
+		{
+			final boolean indeterminate = !NotificationDescription.ShouldShowPercentage || NotificationDescription.Indeterminate;
+			String notificationText;
+			if (NotificationDescription.CurrentProgress < NotificationDescription.MAX_PROGRESS)
+			{
+				notificationText = "";
+			}
+			else
+			{
+				notificationText = NotificationDescription.ContentCompleteText;
+			}
+
+			NotificationCompat.Builder builder = new NotificationCompat.Builder(context, NotificationDescription.NotificationChannelID)
+				.setContentTitle(NotificationDescription.NoInternetAvailable)
+				.setTicker(NotificationDescription.TitleText)
+				.setContentText(notificationText)
+				.setContentIntent(PendingNotificationIntent)
+				.setProgress(NotificationDescription.MAX_PROGRESS, NotificationDescription.CurrentProgress, indeterminate)
+				.setOngoing(true)
+				.setOnlyAlertOnce(true)
+				.setSmallIcon(NotificationDescription.SmallIconResourceID)
+				.setSilent(true);
+			if (!IsInForeground())
+			{
+				builder.addAction(NotificationDescription.CancelIconResourceID, NotificationDescription.CancelText, PendingCancelIntent);
+			}
+			return builder.build();
+		}
+
+		@NonNull Notification CreateCellularWaitNotification(@NonNull Context context)
+		{
+			final boolean indeterminate = !NotificationDescription.ShouldShowPercentage || NotificationDescription.Indeterminate;
+			String notificationText;
+			if (NotificationDescription.CurrentProgress < NotificationDescription.MAX_PROGRESS)
+			{
+				notificationText = "";
+			}
+			else
+			{
+				notificationText = NotificationDescription.ContentCompleteText;
+			}
+
+			NotificationCompat.Builder builder = new NotificationCompat.Builder(context, NotificationDescription.NotificationChannelID)
+				.setContentTitle(NotificationDescription.WaitingForCellularText)
+				.setTicker(NotificationDescription.WaitingForCellularText)
+				.setContentText(notificationText)
+				.setContentIntent(PendingNotificationIntent)
+				.setProgress(NotificationDescription.MAX_PROGRESS, NotificationDescription.CurrentProgress, indeterminate)
+				.setOngoing(true)
+				.setOnlyAlertOnce(true)
+				.setSmallIcon(NotificationDescription.SmallIconResourceID)
+				.setSilent(true);
+			if (!IsInForeground())
+			{
+				builder.addAction(NotificationDescription.CancelIconResourceID, NotificationDescription.ApproveText, PendingApproveIntent);
+				builder.addAction(NotificationDescription.CancelIconResourceID, NotificationDescription.CancelText, PendingCancelIntent);
+			}
+			return builder.build();
+		}
+		
+		@NonNull Notification CreateDownloadProgressNotification(@NonNull Context context)
+		{
+			final boolean indeterminate = !NotificationDescription.ShouldShowPercentage || NotificationDescription.Indeterminate;
+			String notificationText;
+			if (NotificationDescription.CurrentProgress < NotificationDescription.MAX_PROGRESS)
+			{
+				if (indeterminate)
+				{
+					notificationText = NotificationDescription.ContentText.replace("%3d%%", "");
+				}
+				else
+				{
+					notificationText = String.format(NotificationDescription.ContentText, NotificationDescription.CurrentProgress);
+				}
+			}
+			else
+			{
+				notificationText = NotificationDescription.ContentCompleteText;
+			}
+
+			NotificationCompat.Builder builder = new NotificationCompat.Builder(context, NotificationDescription.NotificationChannelID)
+				.setContentTitle(NotificationDescription.TitleText)
+				.setTicker(NotificationDescription.TitleText)
+				.setContentText(notificationText)
+				.setContentIntent(PendingNotificationIntent)
+				.setProgress(NotificationDescription.MAX_PROGRESS, NotificationDescription.CurrentProgress, indeterminate)
+				.setOngoing(true)
+				.setOnlyAlertOnce(true)
+				.setSmallIcon(NotificationDescription.SmallIconResourceID)
+				.setSilent(true);
+			if (!IsInForeground() || GameActivity.IsActivityPaused())
+			{
+				builder.addAction(NotificationDescription.CancelIconResourceID, NotificationDescription.CancelText, PendingCancelIntent);
+			}
+			return builder.build();
+		}
+
+		@MainThread
+		@NonNull Notification CreateNotification(@NonNull Context context)
+		{
+			assert Task != null;
+			
+			//Task.FetchManager.RequestGroupProgressUpdate(0, Task);
+
+			if (Task.bIsNetworkLost && Settings.Global.getInt(context.getContentResolver(), Settings.Global.AIRPLANE_MODE_ON, 0) != 0)
+			{
+				return CreateAirplaneModeNotification(context);
+			}
+			else if (!IsInForeground() && Task.bIsNetworkMetered && Task.RestrictBackgroundStatus == ConnectivityManagerCompat.RESTRICT_BACKGROUND_STATUS_ENABLED)
+			{
+				return CreateDataSaverEnabledNotification(context);
+			}
+			else if (Task.bIsNetworkLost)
+			{
+				return CreateNoInternetDownloadNotification(context);
+			}
+			else if (Task.IsWaitingForCellularApproval())
+			{
+				return CreateCellularWaitNotification(context);
+			}
+			else
+			{
+				return CreateDownloadProgressNotification(context);
+			}
+		}
+		
+		@MainThread
+		public void UpdateProgress(int value)
+		{
+			assert NotificationDescription != null;
+			
+			NotificationDescription.CurrentProgress = value;
+			NotificationDescription.Indeterminate = value < 0;
+		}
+		
+		public void AllComplete(Result result)
+		{
+			Handler.removeCallbacks(this);
+			
+			Completer.set(result);
+		}
+
+		@MainThread
+		public void Cancel()
+		{
+			Log.debug("Cancelling");
+			
+			if (Completer.setCancelled())
+			{
+				Cancelled();
+			}
+		}
+
+		@MainThread
+		void Cancelled()
+		{
+			Handler.removeCallbacks(this);
+
+			if (Task != null && Task.Worker == this)
+			{
+				Task.Worker = null;
+
+				if (!IsInForeground())
+				{
+					Task.Stop();
+				}
+				else
+				{
+					Task.Pause();
+				}
+			}
+		}
+		
+		@MainThread
+		@Override
+		public @NonNull ListenableFuture<Result> startWork()
+		{
+			ListenableFuture<ListenableWorker.Result> future = CallbackToFutureAdapter.getFuture(this);
+
+			final Context context = getApplicationContext();
+			
+			SharedPreferences backgroundPreferences;
+			if (!IsInForeground())
+			{
+				backgroundPreferences = GetBackgroundPreferences(context);
+				
+				FetchManagerInterfaceHelper.InitSharedManager(context);
+
+				try
+				{
+					new Task(context, backgroundPreferences).Restart(backgroundPreferences);
+				}
+				catch (Throwable t)
+				{
+					t.printStackTrace();
+					
+					Completer.setException(t);
+					
+					if (Task != null)
+					{
+						Task.Stop();
+					}
+					
+					return future;
+				}
+			}
+			else
+			{
+				backgroundPreferences = BackgroundPreferences;
+			}
+
+			if (Task != null)
+			{
+				Initialize(context, backgroundPreferences);
+				
+				Task.Worker = this;
+				Task.EnqueueWorkerFuture = null;
+
+				run();
+			}
+			else
+			{
+				Completer.set(Result.success());
+			}
+			
+			return future;
+		}
+
+		@Override
+		public void onStopped()
+		{
+			Log.debug("onStopped");
+			
+			Handler.post(this::Cancelled);
+		}
+		
+		@Override
+		public Object attachCompleter(@NonNull CallbackToFutureAdapter.Completer<ListenableWorker.Result> completer)
+		{
+			assert Completer == null;
+			Completer = completer;
+			return null;
+		}
+
+		@MainThread
+		@Override
+		public void run()
+		{
+			if (Task != null && Task.Worker == this)
+			{
+				@SuppressLint("InlinedApi")
+				ListenableFuture<Void> future = setForegroundAsync(new ForegroundInfo(NotificationDescription.NotificationID, CreateNotification(getApplicationContext()), android.content.pm.ServiceInfo.FOREGROUND_SERVICE_TYPE_DATA_SYNC));
+				future.addListener(() ->
+				{
+					try
+					{
+						future.get();
+						
+						Handler.postDelayed(this, 500);
+					}
+					catch (Throwable t)
+					{
+						if (Build.VERSION.SDK_INT >= Build.VERSION_CODES.S && t instanceof ExecutionException e && e.getCause() instanceof ForegroundServiceStartNotAllowedException)
+						{
+							Log.debug("Not allowed to start as a foreground service");
+						}
+						else
+						{
+							t.printStackTrace();
+							
+							if (Completer.setException(t))
+							{
+								onStopped();
+							}
+						}
+					}
+				}, Runnable::run);
+			}
+		}
+	}
+
+	@Keep
+	public static void Initialize(@NonNull HashMap<String, ?> params, boolean waitTillDone)
+	{
+		Log.debug("Initializing");
+		
+		Handler.post(() ->
+		{
+			assert BackgroundPreferences == null;
+
+			final Context context = GameApplication.getAppContext();
+			
+			if (Task != null)
+			{
+				Task.CancelAll();
+			}
+			else
+			{
+				CancelWorker(context);
+			}
+			
+			assert Task == null;
+
+			synchronized (params)
+			{
+				BackgroundPreferences = GetBackgroundPreferences(context);
+				params.notifyAll();
+			}
+
+			try (Stream<Path> files = Files.list(Paths.get(GetTempDownloadRootPath(BackgroundPreferences))))
+			{
+				files.filter(file -> !file.endsWith("URLMap") && !file.endsWith(GameActivity.BGDL_BUILD_VERSION)).forEach(file ->
+				{
+					Log.debug("cleaning up: " + file.getFileName());
+
+					try
+					{
+						Files.walkFileTree(file, new SimpleFileVisitor<>()
+						{
+							@Override
+							public FileVisitResult visitFile(Path file, BasicFileAttributes attrs)
+							{
+								try
+								{
+									Files.delete(file);
+								}
+								catch (IOException ignored)
+								{
+								}
+								
+								return FileVisitResult.CONTINUE;
+							}
+
+							@Override
+							public FileVisitResult postVisitDirectory(Path dir, IOException exc)
+							{
+								try
+								{
+									Files.delete(dir);
+								}
+								catch (IOException ignored)
+								{
+								}
+								
+								return FileVisitResult.CONTINUE;
+							}
+						});
+					}
+					catch (IOException ignored)
+					{
+					}
+				});
+			}
+			catch (IOException ignored)
+			{
+			}
+
+			SharedPreferences.Editor editor = BackgroundPreferences.edit();
+			for (Map.Entry<String, ?> entry : params.entrySet())
+			{
+				if (entry.getValue() instanceof String value && !value.equals(BackgroundPreferences.getString(entry.getKey(), "")))
+				{
+					editor.putString(entry.getKey(), value);
+				}
+				else if (entry.getValue() instanceof Boolean value && value != BackgroundPreferences.getBoolean(entry.getKey(), false))
+				{
+					editor.putBoolean(entry.getKey(), value);
+				}
+				else if (entry.getValue() instanceof Integer value && value != BackgroundPreferences.getInt(entry.getKey(), 0))
+				{
+					editor.putInt(entry.getKey(), value);
+				}
+			}
+			editor.commit();
+		});
+		
+		if (waitTillDone)
+		{
+			// Not a fan of this. Remove once AsyncDownloader is fixed, or using OkHttp directly
+			synchronized (params)
+			{
+				try
+				{
+					while (BackgroundPreferences == null)
+					{
+						params.wait(250);
+					}
+				}
+				catch (InterruptedException ignored)
+				{
+				}
+			}
+		}
+	}
+
+	@Keep
+	public static void SetMaxConcurrentDownloads(int maxConcurrentDownloads)
+	{
+		BackgroundPreferences.edit().putInt(DownloadWorkerParameterKeys.DOWNLOAD_MAX_CONCURRENT_REQUESTS_KEY, maxConcurrentDownloads).apply();
+	}
+
+	@Keep
+	public static void Enqueue(@NonNull DownloadDescription downloadDescription)
+	{
+		Handler.post(() ->
+		{
+			if (Task == null)
+			{
+				new Task(GameApplication.getAppContext(), BackgroundPreferences).Start();
+			}
+
+			Task.Enqueue(downloadDescription, BackgroundPreferences);
+		});
+	}
+
+	@Keep
+	@MainThread
+	public static void Resume()
+	{
+		if (Task != null && Task.Worker == null)
+		{
+			Task.Resume();
+		}
+	}
+
+	@Keep
+	public static void Pause(@NonNull String destinationLocation)
+	{
+		Handler.post(() ->
+		{
+			if (Task != null)
+			{
+				Task.Pause(destinationLocation);
+			}
+		});
+	}
+	
+	@Keep
+	public static void Resume(@NonNull String destinationLocation)
+	{
+		Handler.post(() ->
+		{
+			if (Task != null)
+			{
+				Task.Resume(destinationLocation);
+			}
+		});
+	}
+	
+	@Keep
+	public static void Cancel(@NonNull String destinationLocation)
+	{
+		Handler.post(() ->
+		{
+			if (Task != null)
+			{
+				Task.Cancel(destinationLocation);
+			}
+		});
+	}
+}
\ No newline at end of file
diff --git a/Engine/Plugins/Online/Android/AndroidFetchBackgroundDownload/Source/AndroidFetchBackgroundDownload/Java/download/DownloadProgressListener.java b/Engine/Plugins/Online/Android/AndroidFetchBackgroundDownload/Source/AndroidFetchBackgroundDownload/Java/download/DownloadProgressListener.java
--- a/Engine/Plugins/Online/Android/AndroidFetchBackgroundDownload/Source/AndroidFetchBackgroundDownload/Java/download/DownloadProgressListener.java
+++ b/Engine/Plugins/Online/Android/AndroidFetchBackgroundDownload/Source/AndroidFetchBackgroundDownload/Java/download/DownloadProgressListener.java
@@ -1,15 +1,18 @@
 // Copyright Epic Games, Inc. All Rights Reserved.
 package com.epicgames.unreal.download;
 
+import androidx.annotation.NonNull;
+
 import com.epicgames.unreal.download.UEDownloadWorker.EDownloadCompleteReason;
+import com.epicgames.unreal.download.datastructs.DownloadDescription;
 
 //Interface for a class that wants to be notified when a download completes
 public interface DownloadProgressListener
 {
-	public void OnDownloadProgress(String DestinationLocation, long TotalBytesWritten, long BytesWrittenSinceLastCall);
-	public void OnDownloadGroupProgress(int GroupID, int Progress, boolean Indeterminate);
-	public void OnDownloadComplete(String DestinationLocation, EDownloadCompleteReason CompleteReason);
-	public void OnDownloadMetrics(String DestinationLocation, long TotalBytesDownloaded, long DownloadDuration, long DownloadStartTimeUTC, long DownloadEndTimeUTC);
-	public void OnAllDownloadsComplete(boolean bDidAllRequestsSucceed);
-	public void OnDownloadEnqueued(String DestinationLocation, boolean bEnqueueSuccess);
+	void OnDownloadProgress(@NonNull String DestinationLocation, long TotalBytesWritten, long BytesWrittenSinceLastCall);
+	void OnDownloadGroupProgress(int GroupID, int Progress, boolean Indeterminate);
+	void OnDownloadComplete(@NonNull String DestinationLocation, EDownloadCompleteReason CompleteReason);
+	void OnDownloadMetrics(@NonNull String DestinationLocation, long TotalBytesDownloaded, long DownloadDuration, long DownloadStartTimeUTC, long DownloadEndTimeUTC);
+	void OnAllDownloadsComplete(boolean bDidAllRequestsSucceed);
+	void OnDownloadEnqueued(@NonNull DownloadDescription DownloadDescription);
 }
\ No newline at end of file
diff --git a/Engine/Plugins/Online/Android/AndroidFetchBackgroundDownload/Source/AndroidFetchBackgroundDownload/Java/download/UEDownloadWorker.java b/Engine/Plugins/Online/Android/AndroidFetchBackgroundDownload/Source/AndroidFetchBackgroundDownload/Java/download/UEDownloadWorker.java
--- a/Engine/Plugins/Online/Android/AndroidFetchBackgroundDownload/Source/AndroidFetchBackgroundDownload/Java/download/UEDownloadWorker.java
+++ b/Engine/Plugins/Online/Android/AndroidFetchBackgroundDownload/Source/AndroidFetchBackgroundDownload/Java/download/UEDownloadWorker.java
@@ -26,6 +26,7 @@ import java.io.File;
 import com.epicgames.unreal.GameActivity;
 
 import com.epicgames.unreal.Logger;
+import com.epicgames.unreal.download.datastructs.DownloadDescription;
 import com.epicgames.unreal.workmanager.UEWorker;
 
 import com.epicgames.unreal.download.datastructs.DownloadNotificationDescription;
@@ -157,7 +158,7 @@ public class UEDownloadWorker extends UEWorker implements DownloadProgressListen
 		{
 			NetworkListener = new NetworkConnectivityClient.Listener() {
 				@Override
-				public void onNetworkAvailable(NetworkConnectivityClient.NetworkTransportType networkTransportType) {
+				public void onNetworkAvailable(NetworkConnectivityClient.NetworkTransportType networkTransportType, boolean isNetworkMetered) {
 					bLostNetwork = false;
 					if (mFetchManager != null) {
 						mFetchManager.NetworkAvailable();
@@ -782,7 +783,7 @@ public class UEDownloadWorker extends UEWorker implements DownloadProgressListen
 	// DownloadCompletionListener Implementation
 	//
 	@Override
-	public void OnDownloadProgress(String DestinationLocation, long TotalBytesWritten, long BytesWrittenSinceLastCall)
+	public void OnDownloadProgress(@NonNull String DestinationLocation, long TotalBytesWritten, long BytesWrittenSinceLastCall)
 	{
 		if (!bGameThreadIsActive)
 		{
@@ -812,7 +813,7 @@ public class UEDownloadWorker extends UEWorker implements DownloadProgressListen
 	}
 	
 	@Override
-	public void OnDownloadComplete(String DestinationLocation, EDownloadCompleteReason CompleteReason)
+	public void OnDownloadComplete(@NonNull String DestinationLocation, EDownloadCompleteReason CompleteReason)
 	{
 		if (!bGameThreadIsActive)
 		{
@@ -831,7 +832,7 @@ public class UEDownloadWorker extends UEWorker implements DownloadProgressListen
 	}
 
 	@Override
-	public void OnDownloadMetrics(String DestinationLocation, long TotalBytesDownloaded, long DownloadDuration, long DownloadStartTimeUTC, long DownloadEndTimeUTC)
+	public void OnDownloadMetrics(@NonNull String DestinationLocation, long TotalBytesDownloaded, long DownloadDuration, long DownloadStartTimeUTC, long DownloadEndTimeUTC)
 	{
 		if (!bGameThreadIsActive)
 		{
@@ -900,17 +901,9 @@ public class UEDownloadWorker extends UEWorker implements DownloadProgressListen
 	}
 
 	@Override
-	public void OnDownloadEnqueued(String DestinationLocation, boolean bEnqueueSuccess)
+	public void OnDownloadEnqueued(@NonNull DownloadDescription DownloadDescription)
 	{
-		if (bEnqueueSuccess)
-		{
-			Log.verbose("Enqueue success:" + DestinationLocation);
-		}
-		else
-		{
-			Log.debug("Enqueue failure, retrying request:" + DestinationLocation);
-			mFetchManager.RetryDownload(DestinationLocation);
-		}
+		Log.verbose("Enqueue success:" + DownloadDescription.DestinationLocation);
 
 		bHasEnqueueHappened = true;
 	}
diff --git a/Engine/Plugins/Online/Android/AndroidFetchBackgroundDownload/Source/AndroidFetchBackgroundDownload/Java/download/asyncdownloader/AsyncDownloadRequest.java b/Engine/Plugins/Online/Android/AndroidFetchBackgroundDownload/Source/AndroidFetchBackgroundDownload/Java/download/asyncdownloader/AsyncDownloadRequest.java
--- a/Engine/Plugins/Online/Android/AndroidFetchBackgroundDownload/Source/AndroidFetchBackgroundDownload/Java/download/asyncdownloader/AsyncDownloadRequest.java
+++ b/Engine/Plugins/Online/Android/AndroidFetchBackgroundDownload/Source/AndroidFetchBackgroundDownload/Java/download/asyncdownloader/AsyncDownloadRequest.java
@@ -13,19 +13,19 @@ import java.util.Objects;
 final public class AsyncDownloadRequest {
     public final String Id;								// internal unique Id used for tracking
     public final String RequestId;						// optional request Id used for deduplication
-    public final List<String> Urls;						// cycle through these on failure
+    public final String[] Urls;							// cycle through these on failure
     public final int MaxRetries;						// total retry attempts across all URLs
     public final File TargetFile;
     public final boolean bResumeIfPossible;				// use Range: bytes=<currentLen>- if file exists, enabled by default
     public final @Nullable Map<String, String> Headers; // optional request headers per attempt
 	public final int GroupID;
-	public final long TotalSize;
+	public long TotalSize;
 
     private AsyncDownloadRequest(Builder b) {
         this.Id = Objects.requireNonNull(b.Id, "id");
         this.RequestId = b.RequestId;
-        this.Urls = List.copyOf(Objects.requireNonNull(b.Urls, "urls"));
-        if (this.Urls.isEmpty()) 
+        this.Urls = b.Urls;
+        if (this.Urls.length == 0) 
 			throw new IllegalArgumentException("urls must not be empty");
         this.MaxRetries = Math.max(0, b.MaxRetries);
         this.TargetFile = Objects.requireNonNull(b.TargetFile, "targetFile");
@@ -36,13 +36,13 @@ final public class AsyncDownloadRequest {
     }
 
     public static class Builder {
-        private String Id;
-        private String RequestId;
-        private List<String> Urls = new ArrayList<>();
-        private int MaxRetries = 3;
-        private File TargetFile;
-        private boolean bResumeIfPossible = true;
-        private Map<String, String> Headers;
+		private String Id;
+		private String RequestId;
+		private String[] Urls = new String[0];
+		private int MaxRetries = 3;
+		private File TargetFile;
+		private boolean bResumeIfPossible = true;
+		private Map<String, String> Headers;
 		private int GroupID = 0;
 		private long TotalSize = -1L;
 
@@ -54,8 +54,8 @@ final public class AsyncDownloadRequest {
 			this.RequestId = RequestId; 
 			return this; 
 		}
-        public Builder SetUrls(List<String> Urls) { 
-			this.Urls = new ArrayList<>(Urls); 
+        public Builder SetUrls(String[] Urls) { 
+			this.Urls = Urls; 
 			return this; 
 		}
 
diff --git a/Engine/Plugins/Online/Android/AndroidFetchBackgroundDownload/Source/AndroidFetchBackgroundDownload/Java/download/asyncdownloader/AsyncDownloadRequestListener.java b/Engine/Plugins/Online/Android/AndroidFetchBackgroundDownload/Source/AndroidFetchBackgroundDownload/Java/download/asyncdownloader/AsyncDownloadRequestListener.java
--- a/Engine/Plugins/Online/Android/AndroidFetchBackgroundDownload/Source/AndroidFetchBackgroundDownload/Java/download/asyncdownloader/AsyncDownloadRequestListener.java
+++ b/Engine/Plugins/Online/Android/AndroidFetchBackgroundDownload/Source/AndroidFetchBackgroundDownload/Java/download/asyncdownloader/AsyncDownloadRequestListener.java
@@ -2,15 +2,15 @@
 
 package com.epicgames.unreal.download.asyncdownloader;
 
-import androidx.annotation.Nullable;
+import androidx.annotation.NonNull;
 
 import java.io.File;
 import java.net.URL;
 
 public interface AsyncDownloadRequestListener {
-    public void OnStart(String Id, URL Url, long StartingOffset, long TotalBytes);
-	public void OnProgress(String Id, long BytesDownloadedSoFar, long TotalBytes);
-	public void OnSuccess(String Id, File File);
-	public void OnError(String Id, Throwable Error, @Nullable URL LastTriedUrl);
-	public void OnPaused(String Id);
+    void OnStart(@NonNull String Id);
+	void OnProgress(@NonNull String Id, long BytesDownloadedSoFar);
+	void OnSuccess(@NonNull String Id, @NonNull File File);
+	void OnError(@NonNull String Id, @NonNull Throwable Error, URL LastTriedUrl);
+	void OnPaused(@NonNull String Id);
 }
diff --git a/Engine/Plugins/Online/Android/AndroidFetchBackgroundDownload/Source/AndroidFetchBackgroundDownload/Java/download/asyncdownloader/AsyncDownloader.java b/Engine/Plugins/Online/Android/AndroidFetchBackgroundDownload/Source/AndroidFetchBackgroundDownload/Java/download/asyncdownloader/AsyncDownloader.java
--- a/Engine/Plugins/Online/Android/AndroidFetchBackgroundDownload/Source/AndroidFetchBackgroundDownload/Java/download/asyncdownloader/AsyncDownloader.java
+++ b/Engine/Plugins/Online/Android/AndroidFetchBackgroundDownload/Source/AndroidFetchBackgroundDownload/Java/download/asyncdownloader/AsyncDownloader.java
@@ -2,6 +2,7 @@
 
 package com.epicgames.unreal.download.asyncdownloader;
 
+import androidx.annotation.NonNull;
 import androidx.annotation.Nullable;
 
 import okhttp3.*;
@@ -55,26 +56,17 @@ public class AsyncDownloader {
 		//Log.debug("Submit: Id=" + Request.Id);
 		DownloadEntry Existing = Entries.get(Request.Id);
 		if (Existing == null) {
-			if (Request.RequestId == null || Request.RequestId.isEmpty()) {
-				DownloadEntry Entry = new DownloadEntry(Request, Listener);
-				Entries.put(Request.Id, Entry);
-				StartOrQueue(Entry);
-				return true;
-			}
-			Existing = Entries.values().stream().filter(Entry -> Request.RequestId.equals(Entry.Request.RequestId)).findFirst().orElse(null);
-			if (Existing == null) {
-				DownloadEntry Entry = new DownloadEntry(Request, Listener);
-				Entries.put(Request.Id, Entry);
-				StartOrQueue(Entry);
-				return true;
-			}
+			DownloadEntry Entry = new DownloadEntry(Request, Listener);
+			Entries.put(Request.Id, Entry);
+			StartOrQueue(Entry);
+			return true;
 		}
 		Existing.Listener = Listener;
 		//Log.debug("Submit: pre-Existing, updated Listener");
 		return false;
 	}
 
-	public synchronized void Pause(String Id) {
+	public synchronized void Pause(String Id, boolean userPaused) {
 		//Log.debug("Pause: Id=" + Id);
 		DownloadEntry Entry = Entries.get(Id);
 		if (Entry == null) {
@@ -84,7 +76,10 @@ public class AsyncDownloader {
 			//Log.debug("Pause: Id=" + Id + " is already paused, ignoring call");
 			return;
 		}
-		Entry.UserPaused = true;
+		if (userPaused)
+		{
+			Entry.UserPaused = true;
+		}
 		Call c = Entry.currentCall;
 		if (c != null) {
 			c.cancel();
@@ -101,11 +96,11 @@ public class AsyncDownloader {
 		//Log.debug("PauseAll");
 		PausedAll.set(true);
 		for (DownloadEntry Entry : Entries.values()) {
-			Pause(Entry.Request.Id);
+			Pause(Entry.Request.Id, false);
 		}
 	}
 
-	public synchronized void Resume(String Id) {
+	public synchronized void Resume(String Id, boolean userResumed) {
 		//Log.debug("Resume: Id=" + Id);
 		DownloadEntry Entry = Entries.get(Id);
 		if (Entry == null) {
@@ -115,7 +110,10 @@ public class AsyncDownloader {
 			//Log.debug("Resume: Id=" + Id + " not paused, ignoring call");
 			return;
 		}
-		Entry.UserPaused = false;
+		if (userResumed)
+		{
+			Entry.UserPaused = false;
+		}
 		StartOrQueue(Entry);
 	}
 
@@ -126,7 +124,7 @@ public class AsyncDownloader {
 		//Log.debug("ResumeAll");
 		PausedAll.set(false);
 		for (DownloadEntry Entry : Entries.values()) {
-			Resume(Entry.Request.Id);
+			Resume(Entry.Request.Id, false);
 		}
 	}
 
@@ -161,7 +159,7 @@ public class AsyncDownloader {
 	public synchronized void Retry(String Id) {
 		//Log.debug("Retry: Id=" + Id);
 
-		DownloadEntry OldEntry = Entries.get(Id);
+		DownloadEntry OldEntry = Entries.remove(Id);
 		if (OldEntry == null) {
 			return;
 		}
@@ -176,7 +174,6 @@ public class AsyncDownloader {
 		}
 		OldEntry.State = EState.CANCELLED;
 
-		Entries.remove(Id);
 		DownloadEntry NewEntry = new DownloadEntry(Req, Listener);
 		NewEntry.UserPaused = false;
 		Entries.put(Req.Id, NewEntry);
@@ -194,7 +191,7 @@ public class AsyncDownloader {
 				continue;
 			}
 			if (e.Request.TotalSize <= 0) {
-				continue;
+				return -1f;
 			}
 
 			if (e.State == EState.COMPLETED) {
@@ -247,7 +244,7 @@ public class AsyncDownloader {
 			Files.createDirectories(file.getParentFile().toPath());
 			long CurrentLen = (file.exists() && Entry.Request.bResumeIfPossible) ? file.length() : 0L;
 
-			URL NextUrl = new URL(Entry.Request.Urls.get(Entry.UrlIndex % Entry.Request.Urls.size()));
+			URL NextUrl = new URL(Entry.Request.Urls[Entry.UrlIndex % Entry.Request.Urls.length]);
 			Request.Builder rb = new Request.Builder().url(NextUrl);
 			if (Entry.Request.Headers != null) {
 				for (Map.Entry<String, String> h : Entry.Request.Headers.entrySet()) {
@@ -255,11 +252,25 @@ public class AsyncDownloader {
 				}
 			}
 
-			if (CurrentLen > 0 && Entry.Request.TotalSize > 0) {
-				if (CurrentLen <= Entry.Request.TotalSize) {
+			if (CurrentLen > 0) {
+				if (CurrentLen < Entry.Request.TotalSize)
+				{
 					rb.header("Range", "bytes=" + CurrentLen + "-");
 					//Log.debug("StartOrQueue: requesting partial content for id " + Entry.Request.Id + " | Offset " + CurrentLen);
-				} else {
+				} else if (CurrentLen == Entry.Request.TotalSize) {
+					synchronized(this) {
+						if (Entries.remove(Entry.Request.Id) != null)
+						{
+							Entry.State = EState.COMPLETED;
+							Entry.DownloadedSoFar = CurrentLen;
+							Entry.Listener.OnStart(Entry.Request.Id);
+							Entry.Listener.OnProgress(Entry.Request.Id, CurrentLen);
+							Entry.Listener.OnSuccess(Entry.Request.Id, file);
+						}
+					}
+					CheckIfAllDownloadsCompleted();
+					return;
+				} else if (Entry.Request.TotalSize > 0) {
 					// be on the safe side: delete the file completely if it is bigger than we expected
 					Log.error("StartOrQueue: requesting full file because current file size is too big (" + CurrentLen + " > " + Entry.Request.TotalSize + ")");
 					CurrentLen = 0;
@@ -270,17 +281,18 @@ public class AsyncDownloader {
 			Request Request = rb.build();
 			Entry.State = EState.RUNNING;
 			Entry.DownloadedSoFar = CurrentLen;
-			Entry.Listener.OnStart(Entry.Request.Id, NextUrl, CurrentLen, Entry.Request.TotalSize);
+			Entry.Listener.OnStart(Entry.Request.Id);
 
 			Call Call = Client.newCall(Request);
 			Entry.currentCall = Call;
+			final Object syncRoot = this;
 			Call.enqueue(new Callback() {
-				@Override public void onFailure(Call Call, IOException Except) {
-					Log.error("onFailure: error: failed to begin downloading | Id=" + Entry.Request.Id);
+				@Override public void onFailure(@NonNull Call Call, @NonNull IOException Except) {
+					Log.error("onFailure: error: failed to begin downloading | Id=" + Entry.Request.Id + " | Exception=" + Except);
 					HandleAttemptFailure(Entry, Except, NextUrl);
 				}
 
-				@Override public void onResponse(Call Call, Response Response) {
+				@Override public void onResponse(@NonNull Call Call, @NonNull Response Response) {
 					//Log.debug("onResponse: begin downloading | Id=" + Entry.Request.Id);
 					try (Response response = Response) {
 						int Code = response.code();
@@ -312,12 +324,13 @@ public class AsyncDownloader {
 						}
 
 						// check if the total matches what we expect
-						long TotalContentLen = 0;
-						long FileOnDiskLen = 0;
+						long FileOnDiskLen;
+						long TotalContentLen;
 						if (bShouldAppend) {
 							FileOnDiskLen = file.length();
-							TotalContentLen = FileOnDiskLen + ContentLen;
+							TotalContentLen = ContentLen >= 0 ? FileOnDiskLen + ContentLen : ContentLen;
 						} else {
+							FileOnDiskLen = 0;
 							TotalContentLen = ContentLen;
 						}
 						if (TotalContentLen != Entry.Request.TotalSize && Entry.Request.TotalSize > 0) {
@@ -327,50 +340,42 @@ public class AsyncDownloader {
 									NextUrl);
 							return;
 						}
+						if (TotalContentLen > 0) {
+							Entry.Request.TotalSize = TotalContentLen;
+						}
 
 						// all is good, start downloading
 						try (BufferedSource Source = Objects.requireNonNull(response.body()).source();
-							 FileChannel Channel = FileChannel.open(file.toPath(),
-									 StandardOpenOption.CREATE,
-									 StandardOpenOption.WRITE)) {
+							 FileChannel Channel = FileChannel.open(file.toPath(), StandardOpenOption.CREATE, StandardOpenOption.WRITE)) {
 
-							long WritePos = bShouldAppend ? file.length() : 0L;
+							long WritePos = FileOnDiskLen;
 							if (!bShouldAppend) {
 								Channel.truncate(0L);
 							}
 
-							byte[] Buffer = new byte[8192];
-							long TotalExpected = (ContentLen >= 0)
-									? (bShouldAppend ? WritePos + ContentLen : ContentLen)
-									: -1L;
+							while (!Source.exhausted()) {
+								long NumBytesRead = Channel.transferFrom(Source, WritePos, 64 * 1024);
 
-							while (true) {
-								int NumBytesRead = Source.read(Buffer, 0, Buffer.length);
-								if (NumBytesRead == -1) {
-									break;
-								}
-
-								Channel.position(WritePos);
-								Channel.write(java.nio.ByteBuffer.wrap(Buffer, 0, NumBytesRead));
 								WritePos += NumBytesRead;
 								Entry.DownloadedSoFar = WritePos;
-								Entry.Listener.OnProgress(Entry.Request.Id, WritePos, TotalExpected);
-							}
-
-							//Log.debug("onResponse: download completed | Id=" + Entry.Request.Id);
-							synchronized(this) {
-								if (Entries.remove(Entry.Request.Id) != null)
-								{ 
-									Entry.State = EState.COMPLETED;
-									Entry.currentCall = null;
-									Entry.Listener.OnSuccess(Entry.Request.Id, file);
-								}
+								Entry.Listener.OnProgress(Entry.Request.Id, WritePos);
 							}
-							CheckIfAllDownloadsCompleted();
 						}
 					} catch (Throwable Error) {
 						HandleAttemptFailure(Entry, Error, NextUrl);
+						return;
 					}
+
+					//Log.debug("onResponse: download completed | Id=" + Entry.Request.Id);
+					synchronized(syncRoot) {
+						if (Entries.remove(Entry.Request.Id) != null)
+						{
+							Entry.State = EState.COMPLETED;
+							Entry.currentCall = null;
+							Entry.Listener.OnSuccess(Entry.Request.Id, file);
+						}
+					}
+					CheckIfAllDownloadsCompleted();
 				}
 			});
 
@@ -379,7 +384,7 @@ public class AsyncDownloader {
 		}
 	}
 
-	private synchronized void HandleAttemptFailure(DownloadEntry Entry, Throwable Error, @Nullable URL LastUrl) {
+	private synchronized void HandleAttemptFailure(@NonNull DownloadEntry Entry, @NonNull Throwable Error, URL LastUrl) {
 		//Log.debug("HandleAttemptFailure: Id=" + Entry.Request.Id);
 		Entry.currentCall = null;
 		if (Entry.UserPaused || PausedAll.get()) {
@@ -391,7 +396,7 @@ public class AsyncDownloader {
 		}
 
 		int Attempt = ++Entry.Attempts;
-		Entry.UrlIndex = (Entry.UrlIndex + 1) % Entry.Request.Urls.size();
+		Entry.UrlIndex = (Entry.UrlIndex + 1) % Entry.Request.Urls.length;
 
 		final boolean bHaveRetriesRemaining = Attempt <= Entry.Request.MaxRetries;
 		if (bHaveRetriesRemaining) {
@@ -400,9 +405,11 @@ public class AsyncDownloader {
 		} else {
 			//Log.debug("HandleAttemptFailure: Id=" + Entry.Request.Id + " | out of retries: failing download");
 
-			Entry.State = EState.FAILED;
-			Entry.Listener.OnError(Entry.Request.Id, Error, LastUrl);
-			Entries.remove(Entry.Request.Id);
+			if (Entries.remove(Entry.Request.Id) != null)
+			{
+				Entry.State = EState.FAILED;
+				Entry.Listener.OnError(Entry.Request.Id, Error, LastUrl);
+			}
 			CheckIfAllDownloadsCompleted();
 		}
 	}
diff --git a/Engine/Plugins/Online/Android/AndroidFetchBackgroundDownload/Source/AndroidFetchBackgroundDownload/Java/download/datastructs/DownloadDescription.java b/Engine/Plugins/Online/Android/AndroidFetchBackgroundDownload/Source/AndroidFetchBackgroundDownload/Java/download/datastructs/DownloadDescription.java
--- a/Engine/Plugins/Online/Android/AndroidFetchBackgroundDownload/Source/AndroidFetchBackgroundDownload/Java/download/datastructs/DownloadDescription.java
+++ b/Engine/Plugins/Online/Android/AndroidFetchBackgroundDownload/Source/AndroidFetchBackgroundDownload/Java/download/datastructs/DownloadDescription.java
@@ -12,8 +12,9 @@ import java.io.IOException;
 import java.io.FileWriter;
 import java.io.FileReader;
 
-import java.util.List;
 import java.util.ArrayList;
+import java.util.List;
+import java.util.stream.IntStream;
 
 import com.epicgames.unreal.download.DownloadProgressListener;
 
@@ -23,7 +24,7 @@ public class DownloadDescription
 	//
 	//Serialized Members -- Pulled from WorkerParameters
 	//
-	public List<String> URLs;
+	public String[] URLs;
 	public String RequestID = null;
 	public String DestinationLocation = null;
 	public int RequestPriority = 0;
@@ -61,9 +62,26 @@ public class DownloadDescription
 	public long DownloadPauseTime = 0;
 	public long TotalPausedTime = 0;
 
+	// Fields used for final metrics
+	public long DownloadDuration = 0;
+	public long DownloadStartTimeUTC = 0;
+	public long DownloadEndTimeUTC = 0;
+	
 	public DownloadDescription()
 	{
-		URLs = new ArrayList<String>();
+	}
+
+	public DownloadDescription(String RequestID, String DestinationLocation, int MaxRetryCount, int IndividualURLRetryCount, int RequestPriority, int GroupID, String[] URLs, boolean bIsPaused, long TotalBytesNeeded)
+	{
+		this.RequestID = RequestID;
+		this.DestinationLocation = DestinationLocation;
+		this.MaxRetryCount = MaxRetryCount;
+		this.IndividualURLRetryCount = IndividualURLRetryCount;
+		this.RequestPriority = RequestPriority;
+		this.GroupID = GroupID;
+		this.URLs = URLs;
+		this.bIsPaused = bIsPaused;
+		this.TotalBytesNeeded = TotalBytesNeeded;
 	}
 	
 	public static ArrayList<DownloadDescription> BuildDescriptionArray()
@@ -221,12 +239,7 @@ public class DownloadDescription
 			if (false == jObject.isNull(URLKey))
 			{
 				org.json.JSONArray JSONURLs = jObject.getJSONArray(URLKey);
-				
-				int NumURLs = JSONURLs.length();
-				for (int URLIndex = 0; URLIndex < NumURLs; ++URLIndex)
-				{
-					ReturnedDescription.URLs.add(JSONURLs.optString(URLIndex));
-				}
+				ReturnedDescription.URLs = IntStream.range(0, JSONURLs.length()).mapToObj(JSONURLs::optString).toArray(String[]::new);
 			}
 
 			if (false == jObject.isNull(bHasCompletedKey))
@@ -262,8 +275,8 @@ public class DownloadDescription
 		try 
 		{
 			org.json.JSONArray URLArray = new org.json.JSONArray();
-			for (int URLIndex = 0; URLIndex < URLs.size(); ++URLIndex) {
-				URLArray.put(URLs.get(URLIndex));
+			for (String URL : URLs) {
+				URLArray.put(URL);
 			}
 
 			JSONObj.put(URLKey, URLArray);
@@ -318,13 +331,13 @@ public class DownloadDescription
 	//Requires the URLs lists to match exactly
 	public boolean HasExactSameURLs(DownloadDescription Other)
 	{
-		return ((this.URLs.size() == Other.URLs.size()) && (!this.HasNewURLs(Other)));
+		return ((this.URLs.length == Other.URLs.length) && (!this.HasNewURLs(Other)));
 	}
 	
 	//This only requires that our DestinationLocation has all the same URLs as the Other, but it can have more.
 	public boolean HasNewURLs(DownloadDescription Other)
 	{
-		return (!this.URLs.containsAll(Other.URLs));
+		return !List.of(this.URLs).containsAll(List.of(Other.URLs));
 	}
 	
 	public boolean HasNewDestinationLocation(DownloadDescription Other)
diff --git a/Engine/Plugins/Online/Android/AndroidFetchBackgroundDownload/Source/AndroidFetchBackgroundDownload/Java/download/datastructs/DownloadNotificationDescription.java b/Engine/Plugins/Online/Android/AndroidFetchBackgroundDownload/Source/AndroidFetchBackgroundDownload/Java/download/datastructs/DownloadNotificationDescription.java
--- a/Engine/Plugins/Online/Android/AndroidFetchBackgroundDownload/Source/AndroidFetchBackgroundDownload/Java/download/datastructs/DownloadNotificationDescription.java
+++ b/Engine/Plugins/Online/Android/AndroidFetchBackgroundDownload/Source/AndroidFetchBackgroundDownload/Java/download/datastructs/DownloadNotificationDescription.java
@@ -3,22 +3,32 @@ package com.epicgames.unreal.download.datastructs;
 
 import android.app.NotificationManager;
 import android.content.Context;
-import android.util.Log;
+import android.content.SharedPreferences;
 
 import androidx.annotation.NonNull;
 import androidx.annotation.Nullable;
 import androidx.work.Data;
-import androidx.work.WorkerParameters;
 
 import com.epicgames.unreal.LocalNotificationReceiver;
 import com.epicgames.unreal.Logger;
-import com.epicgames.unreal.download.datastructs.DownloadWorkerParameterKeys;
+
+import java.util.Map;
 
 //Helper class that stores all the needed information for a Notification in one object and handles parsing and caching defaults stored in the WorkerParameters
 public class DownloadNotificationDescription
 {
 	//Constructor that parses good defaults from our Worker InputData
-	public DownloadNotificationDescription(@NonNull Data data,@NonNull Context context, @Nullable Logger Log)
+	public DownloadNotificationDescription(@NonNull Data data, @NonNull Context context, @Nullable Logger Log)
+	{
+		this(data.getKeyValueMap(), context, Log);
+	}
+
+	public DownloadNotificationDescription(@NonNull SharedPreferences preferences, @NonNull Context context, @Nullable Logger Log)
+	{
+		this(preferences.getAll(), context, Log);
+	}
+
+	public DownloadNotificationDescription(@NonNull Map<String, ?> data, @NonNull Context context, @Nullable Logger Log)
 	{
 		//Set values that are just raw defaults that we expect to be manually overriden
 		{
@@ -28,26 +38,26 @@ public class DownloadNotificationDescription
 					
 		//Load notification channel information
 		{
-			NotificationChannelID = data.getString(DownloadWorkerParameterKeys.NOTIFICATION_CHANNEL_ID_KEY);
+			NotificationChannelID = data.get(DownloadWorkerParameterKeys.NOTIFICATION_CHANNEL_ID_KEY) instanceof String value ? value : null;
 			if (null == NotificationChannelID) 
 			{
 				NotificationChannelID = "ue-downloadworker-channel-id";
 			}
 
-			NotificationChannelName = data.getString(DownloadWorkerParameterKeys.NOTIFICATION_CHANNEL_NAME_KEY);
+			NotificationChannelName = data.get(DownloadWorkerParameterKeys.NOTIFICATION_CHANNEL_NAME_KEY) instanceof String value ? value : null;
 			if (null == NotificationChannelName) 
 			{
 				NotificationChannelName = "ue-downloadworker-channel";
 			}
 
-			NotificationChannelImportance = data.getInt(DownloadWorkerParameterKeys.NOTIFICATION_CHANNEL_IMPORTANCE_KEY, NotificationManager.IMPORTANCE_DEFAULT);
+			NotificationChannelImportance = data.get(DownloadWorkerParameterKeys.NOTIFICATION_CHANNEL_IMPORTANCE_KEY) instanceof Integer value ? value : NotificationManager.IMPORTANCE_DEFAULT;
 		}
 		
 		//Load notification base information
 		{
 			//Loads from data or defaults to a random number (that is hopefully unique) as this can NOT be set to 0 for SetForeground notifications!
 
-			NotificationID = data.getInt(DownloadWorkerParameterKeys.NOTIFICATION_ID_KEY, DownloadWorkerParameterKeys.NOTIFICATION_DEFAULT_ID_KEY);
+			NotificationID = data.get(DownloadWorkerParameterKeys.NOTIFICATION_ID_KEY) instanceof Integer value ? value : DownloadWorkerParameterKeys.NOTIFICATION_DEFAULT_ID_KEY;
 
 			if (NotificationID == 0)
 			{
@@ -60,60 +70,62 @@ public class DownloadNotificationDescription
 
 		//Configuration 
 		{
-			ShouldHandleCellular = data.getBoolean(DownloadWorkerParameterKeys.CELLULAR_HANDLE_KEY, false);
+			ShouldHandleCellular = data.get(DownloadWorkerParameterKeys.CELLULAR_HANDLE_KEY) instanceof Boolean value ? value : false;
+
+			ShouldShowPercentage = data.get(DownloadWorkerParameterKeys.SHOW_PERCENTAGE_KEY) instanceof Boolean value ? value : false;
 		}
 
 		//Load notification content information
 		{
-			TitleText = data.getString(DownloadWorkerParameterKeys.NOTIFICATION_CONTENT_TITLE_KEY);
+			TitleText = data.get(DownloadWorkerParameterKeys.NOTIFICATION_CONTENT_TITLE_KEY) instanceof String value ? value : null;
 			if (null == TitleText) 
 			{
 				TitleText = "Downloading";
 			}
 
-			ContentText = data.getString(DownloadWorkerParameterKeys.NOTIFICATION_CONTENT_TEXT_KEY);
+			ContentText = data.get(DownloadWorkerParameterKeys.NOTIFICATION_CONTENT_TEXT_KEY) instanceof String value ? value : null;
 			if (null == ContentText) 
 			{
 				ContentText = "Download in Progress";
 			}
 
-			ContentCompleteText = data.getString(DownloadWorkerParameterKeys.NOTIFICATION_CONTENT_COMPLETE_TEXT_KEY);
+			ContentCompleteText = data.get(DownloadWorkerParameterKeys.NOTIFICATION_CONTENT_COMPLETE_TEXT_KEY) instanceof String value ? value : null;
 			if (null == ContentCompleteText) 
 			{
 				ContentCompleteText = "Complete";
 			}
 			
-			CancelText = data.getString(DownloadWorkerParameterKeys.NOTIFICATION_CONTENT_CANCEL_DOWNLOAD_TEXT_KEY);
+			CancelText = data.get(DownloadWorkerParameterKeys.NOTIFICATION_CONTENT_CANCEL_DOWNLOAD_TEXT_KEY) instanceof String value ? value : null;
 			if (null == CancelText) 
 			{
 				CancelText = "Cancel";
 			}
 
-			NoInternetAvailable = data.getString(DownloadWorkerParameterKeys.NOTIFICATION_CONTENT_NO_INTERNET_TEXT_KEY);
+			NoInternetAvailable = data.get(DownloadWorkerParameterKeys.NOTIFICATION_CONTENT_NO_INTERNET_TEXT_KEY) instanceof String value ? value : null;
 			if (null == NoInternetAvailable)
 			{
 				NoInternetAvailable = "No Internet Available";
 			}
 
-			AirplaneModeText = data.getString(DownloadWorkerParameterKeys.NOTIFICATION_CONTENT_AIRPLANE_MODE_TEXT_KEY);
+			AirplaneModeText = data.get(DownloadWorkerParameterKeys.NOTIFICATION_CONTENT_AIRPLANE_MODE_TEXT_KEY) instanceof String value ? value : null;
 			if (null == AirplaneModeText)
 			{
 				AirplaneModeText = "Airplane Mode Enabled";
 			}
 			
-			DataSaverEnabledText = data.getString(DownloadWorkerParameterKeys.NOTIFICATION_CONTENT_DATA_SAVER_ENABLED_TEXT_KEY);
+			DataSaverEnabledText = data.get(DownloadWorkerParameterKeys.NOTIFICATION_CONTENT_DATA_SAVER_ENABLED_TEXT_KEY) instanceof String value ? value : null;
 			if (null == DataSaverEnabledText)
 			{
 				DataSaverEnabledText = "Data Saver Enabled";
 			}
 
-			WaitingForCellularText = data.getString(DownloadWorkerParameterKeys.NOTIFICATION_CONTENT_WAITING_FOR_CELLULAR_TEXT_KEY);
+			WaitingForCellularText = data.get(DownloadWorkerParameterKeys.NOTIFICATION_CONTENT_WAITING_FOR_CELLULAR_TEXT_KEY) instanceof String value ? value : null;
 			if (null == WaitingForCellularText)
 			{
 				WaitingForCellularText = "Waiting for cellular approval";
 			}
 
-			ApproveText = data.getString(DownloadWorkerParameterKeys.NOTIFICATION_CONTENT_APPROVE_TEXT_KEY);
+			ApproveText = data.get(DownloadWorkerParameterKeys.NOTIFICATION_CONTENT_APPROVE_TEXT_KEY) instanceof String value ? value : null;
 			if (null == ApproveText)
 			{
 				ApproveText = "Approve";
@@ -125,7 +137,7 @@ public class DownloadNotificationDescription
 			//Flag if the user tried to set anything so we can show an error if the load fails when they weren't expecting to use the default
 			boolean bTriedToSetAnything = false;
 			
-			String CancelIconResourceName = data.getString(DownloadWorkerParameterKeys.NOTIFICATION_RESOURCE_CANCEL_ICON_NAME);
+			String CancelIconResourceName = data.get(DownloadWorkerParameterKeys.NOTIFICATION_RESOURCE_CANCEL_ICON_NAME) instanceof String value ? value : null;
 			if (null == CancelIconResourceName) 
 			{
 				CancelIconResourceName = "ic_delete";
@@ -135,7 +147,7 @@ public class DownloadNotificationDescription
 				bTriedToSetAnything = true;
 			}
 
-			String CancelIconResourceType = data.getString(DownloadWorkerParameterKeys.NOTIFICATION_RESOURCE_CANCEL_ICON_TYPE);
+			String CancelIconResourceType = data.get(DownloadWorkerParameterKeys.NOTIFICATION_RESOURCE_CANCEL_ICON_TYPE) instanceof String value ? value : null;
 			if (null == CancelIconResourceType) 
 			{
 				CancelIconResourceType = "drawable";
@@ -145,7 +157,7 @@ public class DownloadNotificationDescription
 				bTriedToSetAnything = true;
 			}
 
-			String CancelIconResourcePackage = data.getString(DownloadWorkerParameterKeys.NOTIFICATION_RESOURCE_CANCEL_ICON_PACKAGE);
+			String CancelIconResourcePackage = data.get(DownloadWorkerParameterKeys.NOTIFICATION_RESOURCE_CANCEL_ICON_PACKAGE) instanceof String value ? value : null;
 			if (null == CancelIconResourcePackage) 
 			{
 				CancelIconResourcePackage = context.getPackageName();
@@ -186,7 +198,7 @@ public class DownloadNotificationDescription
 			//Flag if the user tried to set anything so we can show an error if the load fails when they weren't expecting to use the default
 			boolean bTriedToSetAnything = false;
 			
-			String SmallIconResourceName = data.getString(DownloadWorkerParameterKeys.NOTIFICATION_RESOURCE_SMALL_ICON_NAME);
+			String SmallIconResourceName = data.get(DownloadWorkerParameterKeys.NOTIFICATION_RESOURCE_SMALL_ICON_NAME) instanceof String value ? value : null;
 			if (null == SmallIconResourceName)
 			{
 				SmallIconResourceName = "ic_notification_simple";
@@ -196,7 +208,7 @@ public class DownloadNotificationDescription
 				bTriedToSetAnything = true;
 			}
 
-			String SmallIconResourceType = data.getString(DownloadWorkerParameterKeys.NOTIFICATION_RESOURCE_SMALL_ICON_TYPE);
+			String SmallIconResourceType = data.get(DownloadWorkerParameterKeys.NOTIFICATION_RESOURCE_SMALL_ICON_TYPE) instanceof String value ? value : null;
 			if (null == SmallIconResourceType)
 			{
 				SmallIconResourceType = "drawable";
@@ -206,7 +218,7 @@ public class DownloadNotificationDescription
 				bTriedToSetAnything = true;
 			}
 
-			String SmallIconResourcePackage = data.getString(DownloadWorkerParameterKeys.NOTIFICATION_RESOURCE_SMALL_ICON_PACKAGE);
+			String SmallIconResourcePackage = data.get(DownloadWorkerParameterKeys.NOTIFICATION_RESOURCE_SMALL_ICON_PACKAGE) instanceof String value ? value : null;
 			if (null == SmallIconResourcePackage)
 			{
 				SmallIconResourcePackage = context.getPackageName();
@@ -259,6 +271,8 @@ public class DownloadNotificationDescription
 
 	public boolean ShouldHandleCellular = false;
 
+	public boolean ShouldShowPercentage = false;
+
 	public String TitleText = null;
 	public String ContentText = null;
 	public String ContentCompleteText = null;
diff --git a/Engine/Plugins/Online/Android/AndroidFetchBackgroundDownload/Source/AndroidFetchBackgroundDownload/Java/download/datastructs/DownloadQueueDescription.java b/Engine/Plugins/Online/Android/AndroidFetchBackgroundDownload/Source/AndroidFetchBackgroundDownload/Java/download/datastructs/DownloadQueueDescription.java
--- a/Engine/Plugins/Online/Android/AndroidFetchBackgroundDownload/Source/AndroidFetchBackgroundDownload/Java/download/datastructs/DownloadQueueDescription.java
+++ b/Engine/Plugins/Online/Android/AndroidFetchBackgroundDownload/Source/AndroidFetchBackgroundDownload/Java/download/datastructs/DownloadQueueDescription.java
@@ -2,48 +2,47 @@
 package com.epicgames.unreal.download.datastructs;
 
 import android.content.Context;
-import android.util.Log;
+import android.content.SharedPreferences;
 
 import androidx.annotation.NonNull;
 import androidx.annotation.Nullable;
 import androidx.work.Data;
-import androidx.work.WorkerParameters;
 
 import com.epicgames.unreal.Logger;
-
-import com.epicgames.unreal.download.datastructs.DownloadDescription;
-import com.epicgames.unreal.download.datastructs.DownloadWorkerParameterKeys;
 import com.epicgames.unreal.download.DownloadProgressListener;
 
 import org.json.JSONArray;
 import org.json.JSONException;
 
 import java.util.ArrayList;
-import java.util.List;
+import java.util.Map;
 
 
 //Helper class that is used to serialize any information needed for our downloads as a whole from/to WorkerParameters
 public class DownloadQueueDescription
 {
-	//Constructor that parses the List<DownloadDescription> from our Worker InputData (WorkerParameters)
-	public DownloadQueueDescription(@NonNull Data data, @NonNull Context context, @Nullable Logger Log)
-	{		
-		//Parse DownloadDescriptions 
-		String DownloadDescriptionListFileName = GetDownloadDescriptionListFileName(data, Log);
-		if (DownloadDescriptionListFileName != null)
-		{
-			ParseDownloadDescriptionString(DownloadDescriptionListFileName, Log);
-		}
-		
-		MaxConcurrentDownloads = data.getInt(DownloadWorkerParameterKeys.DOWNLOAD_MAX_CONCURRENT_REQUESTS_KEY, DEFAULT_MAX_CONCURRENT_DOWNLOADS);
+	public final static int DEFAULT_MAX_CONCURRENT_DOWNLOADS = 4;
+
+	public DownloadQueueDescription(ArrayList<DownloadDescription> downloadDescriptions, @NonNull Map<String, ?> data)
+	{
+		DownloadDescriptions = downloadDescriptions;
+		MaxConcurrentDownloads = data.get(DownloadWorkerParameterKeys.DOWNLOAD_MAX_CONCURRENT_REQUESTS_KEY) instanceof Integer value ? value : DEFAULT_MAX_CONCURRENT_DOWNLOADS;
 	}
 
-	public DownloadQueueDescription(String DownloadDescriptionListFileName, @NonNull Data data, @NonNull Context context, @Nullable Logger Log)
+	public DownloadQueueDescription(ArrayList<DownloadDescription> downloadDescriptions, @NonNull SharedPreferences backgroundPreferences)
 	{
-		//Parse DownloadDescriptions 
-		ParseDownloadDescriptionString(DownloadDescriptionListFileName, Log);
+		this(downloadDescriptions, backgroundPreferences.getAll());
+	}
+
+	//Constructor that parses the List<DownloadDescription> from our Worker InputData (WorkerParameters)
+	public DownloadQueueDescription(@NonNull String DownloadDescriptionListFileName, @NonNull Data data, @Nullable Logger Log)
+	{
+		this(ParseDownloadDescriptionString(DownloadDescriptionListFileName, Log), data.getKeyValueMap());
+	}
 
-		MaxConcurrentDownloads = data.getInt(DownloadWorkerParameterKeys.DOWNLOAD_MAX_CONCURRENT_REQUESTS_KEY, DEFAULT_MAX_CONCURRENT_DOWNLOADS);
+	public DownloadQueueDescription(@NonNull String DownloadDescriptionListFileName, @NonNull Data data, @NonNull Context context, @Nullable Logger Log)
+	{
+		this(DownloadDescriptionListFileName, data, Log);
 	}
 
 	public static String GetDownloadDescriptionListFileName(@NonNull Data data, @Nullable Logger Log)
@@ -62,9 +61,9 @@ public class DownloadQueueDescription
 	}
 
 	//Parses our DownloadDescriptions member based on the passed in JSONObject file
-	private void ParseDownloadDescriptionString(@NonNull String DownloadDescriptionListStringFileName, @Nullable Logger Log)
+	static ArrayList<DownloadDescription> ParseDownloadDescriptionString(@NonNull String DownloadDescriptionListStringFileName, @Nullable Logger Log)
 	{
-		DownloadDescriptions = new ArrayList<DownloadDescription>();
+		ArrayList<DownloadDescription> DownloadDescriptions = new ArrayList<>();
 		
 		try 
 		{
@@ -95,12 +94,12 @@ public class DownloadQueueDescription
 				Log.error("Exception while parsing download descriptions: " + e.getMessage());
 			}			
 		}
+		
+		return DownloadDescriptions;
 	}
 
-	public ArrayList<DownloadDescription> DownloadDescriptions = new ArrayList<DownloadDescription>();
+	public ArrayList<DownloadDescription> DownloadDescriptions;
 	public DownloadProgressListener ProgressListener = null;
 	public int DownloadGroupID = 0;
-	public int MaxConcurrentDownloads = 4;
-	
-	private int DEFAULT_MAX_CONCURRENT_DOWNLOADS = 4;
+	public int MaxConcurrentDownloads = DEFAULT_MAX_CONCURRENT_DOWNLOADS;
 }
\ No newline at end of file
diff --git a/Engine/Plugins/Online/Android/AndroidFetchBackgroundDownload/Source/AndroidFetchBackgroundDownload/Java/download/datastructs/DownloadWorkerParameterKeys.java b/Engine/Plugins/Online/Android/AndroidFetchBackgroundDownload/Source/AndroidFetchBackgroundDownload/Java/download/datastructs/DownloadWorkerParameterKeys.java
--- a/Engine/Plugins/Online/Android/AndroidFetchBackgroundDownload/Source/AndroidFetchBackgroundDownload/Java/download/datastructs/DownloadWorkerParameterKeys.java
+++ b/Engine/Plugins/Online/Android/AndroidFetchBackgroundDownload/Source/AndroidFetchBackgroundDownload/Java/download/datastructs/DownloadWorkerParameterKeys.java
@@ -31,6 +31,8 @@ public class DownloadWorkerParameterKeys
 	public final static int NOTIFICATION_DEFAULT_ID_KEY = 1923901283;
 
 	public final static String CELLULAR_HANDLE_KEY = "CellularHandleKey";
+
+	public final static String SHOW_PERCENTAGE_KEY = "bShowPercentage";
 	
 	public final static String NOTIFICATION_CONTENT_TITLE_KEY = "NotificationContentTitle";
 	public final static String NOTIFICATION_CONTENT_TEXT_KEY = "NotificationContentText";
diff --git a/Engine/Plugins/Online/Android/AndroidFetchBackgroundDownload/Source/AndroidFetchBackgroundDownload/Java/download/fetch/FetchManagerInterfaceHelper.java b/Engine/Plugins/Online/Android/AndroidFetchBackgroundDownload/Source/AndroidFetchBackgroundDownload/Java/download/fetch/FetchManagerInterfaceHelper.java
--- a/Engine/Plugins/Online/Android/AndroidFetchBackgroundDownload/Source/AndroidFetchBackgroundDownload/Java/download/fetch/FetchManagerInterfaceHelper.java
+++ b/Engine/Plugins/Online/Android/AndroidFetchBackgroundDownload/Source/AndroidFetchBackgroundDownload/Java/download/fetch/FetchManagerInterfaceHelper.java
@@ -31,6 +31,7 @@ public class FetchManagerInterfaceHelper {
 					int deletedCount = 0;
 					for (File file : files) {
 						if (!file.getName().endsWith("DownloadDescriptionJSONs") &&
+							!file.getName().endsWith("DownloadDescriptions") &&
 							!file.getName().endsWith("URLMap") &&
 							!file.getName().equals(GameActivity.BGDL_BUILD_VERSION)) {
 							Log.debug("cleaning up: " + file.getName());
diff --git a/Engine/Plugins/Online/Android/AndroidFetchBackgroundDownload/Source/AndroidFetchBackgroundDownload/Java/download/fetch/FetchManagerNew.java b/Engine/Plugins/Online/Android/AndroidFetchBackgroundDownload/Source/AndroidFetchBackgroundDownload/Java/download/fetch/FetchManagerNew.java
--- a/Engine/Plugins/Online/Android/AndroidFetchBackgroundDownload/Source/AndroidFetchBackgroundDownload/Java/download/fetch/FetchManagerNew.java
+++ b/Engine/Plugins/Online/Android/AndroidFetchBackgroundDownload/Source/AndroidFetchBackgroundDownload/Java/download/fetch/FetchManagerNew.java
@@ -4,31 +4,22 @@ package com.epicgames.unreal.download.fetch;
 
 import android.content.Context;
 import android.content.SharedPreferences;
-import android.net.Uri;
+
+import androidx.annotation.NonNull;
 
 import com.epicgames.unreal.Logger;
-import com.epicgames.unreal.download.datastructs.DownloadDescription;
-import com.epicgames.unreal.download.datastructs.DownloadQueueDescription;
 import com.epicgames.unreal.download.DownloadProgressListener;
-import com.epicgames.unreal.download.asyncdownloader.AsyncDownloader;
-import com.epicgames.unreal.download.asyncdownloader.AsyncDownloaderListener;
+import com.epicgames.unreal.download.UEDownloadWorker.EDownloadCompleteReason;
 import com.epicgames.unreal.download.asyncdownloader.AsyncDownloadRequest;
 import com.epicgames.unreal.download.asyncdownloader.AsyncDownloadRequestListener;
-
-import com.epicgames.unreal.download.UEDownloadWorker.EDownloadCompleteReason;
-
-import okhttp3.OkHttpClient;
-
-import androidx.annotation.Nullable;
-import androidx.annotation.NonNull;
+import com.epicgames.unreal.download.asyncdownloader.AsyncDownloader;
+import com.epicgames.unreal.download.asyncdownloader.AsyncDownloaderListener;
+import com.epicgames.unreal.download.datastructs.DownloadDescription;
+import com.epicgames.unreal.download.datastructs.DownloadQueueDescription;
 
 import java.io.File;
-import java.util.ArrayList;
-import java.util.HashMap;
-import java.util.List;
-import java.util.Set;
-import java.util.Map;
 import java.net.URL;
+import java.util.List;
 
 public class FetchManagerNew implements FetchManagerInterface
 {
@@ -40,6 +31,11 @@ public class FetchManagerNew implements FetchManagerInterface
 	{
 		return SharedManager;
 	}
+	
+	public Object SyncRoot()
+	{
+		return AsyncDownloader;
+	}
 
 	@Override	
 	public synchronized void StopWork(String WorkID)
@@ -50,20 +46,81 @@ public class FetchManagerNew implements FetchManagerInterface
 			AsyncDownloader.CancelAll();
 		}
 	}
+
+	public synchronized void Enqueue(@NonNull Context context, @NonNull DownloadQueueDescription QueueDescription)
+	{
+		if (AsyncDownloader == null)
+		{
+			SharedPreferences preferences = context.getSharedPreferences("BackgroundDownload", Context.MODE_PRIVATE);
+			boolean AllowHTTP2 = preferences.getBoolean("AllowHTTP2", false);
+			AsyncDownloader = new AsyncDownloader(AllowHTTP2, QueueDescription.MaxConcurrentDownloads);
+			SetupAllDownloadsCompletedListener_Internal();
+		}
+		else
+		{
+			AsyncDownloader.SetMaxConcurrent(QueueDescription.MaxConcurrentDownloads);
+		}
+
+		GlobalDownloadProgressListener = QueueDescription.ProgressListener;
+
+		for (DownloadDescription DownloadDescription : QueueDescription.DownloadDescriptions)
+		{
+			//noinspection SynchronizeOnNonFinalField
+			synchronized (AsyncDownloader)
+			{
+				final String TargetLocation = DownloadDescription.DestinationLocation;
+				final String DownloadLocation = DownloadDescription.DestinationLocation + ".temp";
+				final boolean bFileExists = new File(TargetLocation).exists();
+
+				if (bFileExists)
+				{
+					GlobalDownloadProgressListener.OnDownloadMetrics(DownloadDescription.DestinationLocation, new File(TargetLocation).length(), 0, DownloadDescription.DownloadStartTime, DownloadDescription.DownloadStartTime);
+					GlobalDownloadProgressListener.OnDownloadComplete(DownloadDescription.DestinationLocation, EDownloadCompleteReason.Success);
+
+					AsyncDownloader.Cancel(DownloadDescription.DestinationLocation);
+					continue;
+				}
+
+				final boolean bSuccess = AsyncDownloader.Submit(
+					new AsyncDownloadRequest.Builder()
+						.SetId(DownloadDescription.DestinationLocation)
+						.SetRequestId(DownloadDescription.RequestID)
+						.SetUrls(DownloadDescription.URLs)
+						.SetMaxRetries(DownloadDescription.MaxRetryCount)
+						.SetGroupID(DownloadDescription.GroupID)
+						.SetTargetFile(new File(DownloadLocation))
+						.SetResumeIfPossible(true)
+						.SetTotalSize(DownloadDescription.TotalBytesNeeded)
+						.Build(),
+					new AsyncDownloadRequestListenerImpl(this, GlobalDownloadProgressListener, TargetLocation, DownloadLocation));
+				if (!bSuccess)
+				{
+					RetryDownload(DownloadDescription.DestinationLocation);
+				}
+
+				GlobalDownloadProgressListener.OnDownloadEnqueued(DownloadDescription);
+			}
+		}
+	}
 	
 	@Override
 	public synchronized void EnqueueRequests(@NonNull Context context, @NonNull DownloadQueueDescription QueueDescription)
 	{
 		Log.debug("FetchManagerNew: EnqueueRequests with " + QueueDescription.DownloadDescriptions.size() + " requests");
 
+		if (AsyncDownloader != null)
+		{
+			AsyncDownloader.CancelAll();
+		}
+
 		NumFailedRequests = 0;
-		if (QueueDescription.DownloadDescriptions.size() == 0) {
+		if (QueueDescription.DownloadDescriptions.isEmpty()) {
 			return;
 		}
 
 		if (AsyncDownloader == null)
 		{
-			SharedPreferences preferences = context.getSharedPreferences("BackgroundDownload", context.MODE_PRIVATE);
+			SharedPreferences preferences = context.getSharedPreferences("BackgroundDownload", Context.MODE_PRIVATE);
 			boolean AllowHTTP2 = preferences.getBoolean("AllowHTTP2", false);
 			AsyncDownloader = new AsyncDownloader(AllowHTTP2, QueueDescription.MaxConcurrentDownloads);
 			SetupAllDownloadsCompletedListener_Internal();
@@ -77,22 +134,41 @@ public class FetchManagerNew implements FetchManagerInterface
 
 		for (DownloadDescription DownloadDescription : QueueDescription.DownloadDescriptions)
 		{
-			final String TargetLocation = DownloadDescription.DestinationLocation;
-			final String DownloadLocation = DownloadDescription.DestinationLocation + ".temp";
-
-			final boolean bSuccess = AsyncDownloader.Submit(
-				new AsyncDownloadRequest.Builder()
-					.SetId(DownloadDescription.DestinationLocation)
-					.SetRequestId(DownloadDescription.RequestID)
-					.SetUrls(List.copyOf(DownloadDescription.URLs))
-					.SetMaxRetries(DownloadDescription.MaxRetryCount)
-					.SetGroupID(DownloadDescription.GroupID)
-					.SetTargetFile(new File(DownloadLocation))
-					.SetResumeIfPossible(true)
-					.SetTotalSize(DownloadDescription.TotalBytesNeeded)
-					.Build(),
-				new AsyncDownloadRequestListenerImpl(this, GlobalDownloadProgressListener, TargetLocation, DownloadLocation));
-			GlobalDownloadProgressListener.OnDownloadEnqueued(DownloadDescription.DestinationLocation, bSuccess);
+			//noinspection SynchronizeOnNonFinalField
+			synchronized (AsyncDownloader)
+			{
+				final String TargetLocation = DownloadDescription.DestinationLocation;
+				final String DownloadLocation = DownloadDescription.DestinationLocation + ".temp";
+				final boolean bFileExists = new File(TargetLocation).exists();
+
+				if (bFileExists)
+				{
+					GlobalDownloadProgressListener.OnDownloadMetrics(DownloadDescription.DestinationLocation, new File(TargetLocation).length(), 0, DownloadDescription.DownloadStartTime, DownloadDescription.DownloadStartTime);
+					GlobalDownloadProgressListener.OnDownloadComplete(DownloadDescription.DestinationLocation, EDownloadCompleteReason.Success);
+					
+					AsyncDownloader.Cancel(DownloadDescription.DestinationLocation);
+					continue;
+				}
+				
+				final boolean bSuccess = AsyncDownloader.Submit(
+					new AsyncDownloadRequest.Builder()
+						.SetId(DownloadDescription.DestinationLocation)
+						.SetRequestId(DownloadDescription.RequestID)
+						.SetUrls(DownloadDescription.URLs)
+						.SetMaxRetries(DownloadDescription.MaxRetryCount)
+						.SetGroupID(DownloadDescription.GroupID)
+						.SetTargetFile(new File(DownloadLocation))
+						.SetResumeIfPossible(true)
+						.SetTotalSize(DownloadDescription.TotalBytesNeeded)
+						.Build(),
+					new AsyncDownloadRequestListenerImpl(this, GlobalDownloadProgressListener, TargetLocation, DownloadLocation));
+				if (!bSuccess)
+				{
+					RetryDownload(DownloadDescription.DestinationLocation);
+				}
+
+				GlobalDownloadProgressListener.OnDownloadEnqueued(DownloadDescription);
+			}
 		}
 
 		// clear flag that stops downloader from downloading
@@ -106,9 +182,9 @@ public class FetchManagerNew implements FetchManagerInterface
 			return;
 		}
 		if (bPause) {
-			AsyncDownloader.Pause(DestinationLocation);
+			AsyncDownloader.Pause(DestinationLocation, true);
 		} else {
-			AsyncDownloader.Resume(DestinationLocation);
+			AsyncDownloader.Resume(DestinationLocation, true);
 		}
 	}
 
@@ -218,21 +294,20 @@ public class FetchManagerNew implements FetchManagerInterface
 		}
 	}
 	
-	private static volatile FetchManagerNew SharedManager = new FetchManagerNew();
-	private volatile OkHttpClient OkHttpClient = null;
+	private static final FetchManagerNew SharedManager = new FetchManagerNew();
 	private volatile AllDownloadsCompletedListener AllDownloadsCompletedListener = null;
 	private volatile AsyncDownloader AsyncDownloader = null;
 	private volatile DownloadProgressListener GlobalDownloadProgressListener = null;
 	private volatile int NumFailedRequests = 0;
-	private Logger Log = new Logger("UE", "FetchManagerNew");
+	private final Logger Log = new Logger("UE", "FetchManagerNew");
 
 	// ---------------------------------------------------------------------------------------------------------------------
 	// private helper class
 	// ---------------------------------------------------------------------------------------------------------------------
 
-	final private class AsyncDownloadRequestListenerImpl implements AsyncDownloadRequestListener {
+	static final private class AsyncDownloadRequestListenerImpl implements AsyncDownloadRequestListener {
 
-		public AsyncDownloadRequestListenerImpl(@NonNull FetchManagerNew Owner, @NonNull DownloadProgressListener ProgressListener, String TargetLocation, String DownloadLocation) {
+		public AsyncDownloadRequestListenerImpl(@NonNull FetchManagerNew Owner, @NonNull DownloadProgressListener ProgressListener, @NonNull String TargetLocation, @NonNull String DownloadLocation) {
 			this.Owner = Owner;
 			this.ProgressListener = ProgressListener;
 			this.TargetLocation = TargetLocation;
@@ -240,28 +315,24 @@ public class FetchManagerNew implements FetchManagerInterface
 		}
 
 		@Override
-		public void OnStart(String Id, URL url, long startingOffset, long totalBytesOrMinus1) {
-			LastDownloadedSoFar = startingOffset;
-			TotalBytes = totalBytesOrMinus1;			// total target file size, regardless of starting offset
-
-			if (StartTime == -1L) {
-				StartTime = System.currentTimeMillis();
-			}
-
-			// reporting LastDownloadedSoFar as "bytes downloaded since last call" gives less jumpy progress bars
-			ProgressListener.OnDownloadProgress(Id, LastDownloadedSoFar, LastDownloadedSoFar);	
+		public void OnStart(@NonNull String Id) {
+			StartTime = System.currentTimeMillis();
 		}
 		@Override
-		public void OnProgress(String Id, long bytesDownloadedSoFar, long totalBytesOrMinus1) {
+		public void OnProgress(@NonNull String Id, long bytesDownloadedSoFar) {
 			// Log.debug("OnProgress: " + Id + " | " + bytesDownloadedSoFar + " | " + totalBytesOrMinus1);
-
-			// the Math.max below is only necessary if resuming a download after an error, and the server suddenly stops supporting "206 Partial Content"
-			long TotalDownloadedSinceLastCall = Math.max(0L, bytesDownloadedSoFar - LastDownloadedSoFar);	
 			LastDownloadedSoFar = bytesDownloadedSoFar;
-			ProgressListener.OnDownloadProgress(Id, TotalDownloadedSinceLastCall, LastDownloadedSoFar);
+
+			long NewUpdateTime = System.currentTimeMillis();
+			if (NewUpdateTime >= UpdateTime + 100) {
+				long TotalDownloadedSinceLastUpdate = LastDownloadedSoFar - LastUpdatedSoFar;
+				LastUpdatedSoFar = LastDownloadedSoFar;
+				UpdateTime = NewUpdateTime;
+				ProgressListener.OnDownloadProgress(Id, LastUpdatedSoFar, TotalDownloadedSinceLastUpdate);
+			}
 		}
 		@Override
-		public void OnSuccess(String Id, File file) {
+		public void OnSuccess(@NonNull String Id, @NonNull File file) {
 			EDownloadCompleteReason Reason = EDownloadCompleteReason.Success;
 			File NewFileAtDestinationLocation = new File(TargetLocation);
 			if (!file.renameTo(NewFileAtDestinationLocation)) {
@@ -272,16 +343,16 @@ public class FetchManagerNew implements FetchManagerInterface
 			OnCompleteDownload_Internal(Id, Reason);
 		}
 		@Override
-		public void OnError(String Id, Throwable Error, @Nullable URL LastTriedUrl) {
+		public void OnError(@NonNull String Id, @NonNull Throwable Error, URL LastTriedUrl) {
 			Log.debug("OnError: " + Id);
 			OnCompleteDownload_Internal(Id, EDownloadCompleteReason.OutOfRetries);
 		}
 		@Override
-		public void OnPaused(String Id) {
+		public void OnPaused(@NonNull String Id) {
 			// do we need to do something here?
 		}
 
-		private void OnCompleteDownload_Internal(String Id, EDownloadCompleteReason Reason) {
+		private void OnCompleteDownload_Internal(@NonNull String Id, EDownloadCompleteReason Reason) {
 			long EndTime = System.currentTimeMillis();
 
 			long Duration = 0;
@@ -293,18 +364,24 @@ public class FetchManagerNew implements FetchManagerInterface
 				Owner.NumFailedRequests++;
 			}
 
-			ProgressListener.OnDownloadMetrics(Id, TotalBytes, Duration, StartTime, EndTime);
+			long TotalDownloadedSinceLastUpdate = LastDownloadedSoFar - LastUpdatedSoFar;
+			if (TotalDownloadedSinceLastUpdate > 0)	{
+				LastUpdatedSoFar = LastDownloadedSoFar;
+				UpdateTime = EndTime;
+				ProgressListener.OnDownloadProgress(Id, LastUpdatedSoFar, TotalDownloadedSinceLastUpdate);
+			}
+			ProgressListener.OnDownloadMetrics(Id, LastUpdatedSoFar, Duration, StartTime, EndTime);
 			ProgressListener.OnDownloadComplete(Id, Reason);
 		}
 
 		private long LastDownloadedSoFar = 0;
-		private long TotalBytes = 0;
-		private long StartTime = -1L;
-		private @NonNull DownloadProgressListener ProgressListener = null;
-		private @NonNull FetchManagerNew Owner;
-		private Logger Log = new Logger("UE", "FetchManagerNew");
-		private String TargetLocation;
-		private String DownloadLocation;
+		private long StartTime = -1;
+		private long LastUpdatedSoFar = 0;
+		private long UpdateTime = Long.MIN_VALUE;
+		private final @NonNull DownloadProgressListener ProgressListener;
+		private final @NonNull FetchManagerNew Owner;
+		private final Logger Log = new Logger("UE", "FetchManagerNew");
+		private final String TargetLocation;
+		private final String DownloadLocation;
 	}
-
 }
diff --git a/Engine/Plugins/Online/Android/AndroidFetchBackgroundDownload/Source/AndroidFetchBackgroundDownload/Java/download/fetch/FetchManagerOld.java b/Engine/Plugins/Online/Android/AndroidFetchBackgroundDownload/Source/AndroidFetchBackgroundDownload/Java/download/fetch/FetchManagerOld.java
--- a/Engine/Plugins/Online/Android/AndroidFetchBackgroundDownload/Source/AndroidFetchBackgroundDownload/Java/download/fetch/FetchManagerOld.java
+++ b/Engine/Plugins/Online/Android/AndroidFetchBackgroundDownload/Source/AndroidFetchBackgroundDownload/Java/download/fetch/FetchManagerOld.java
@@ -160,21 +160,9 @@ public class FetchManagerOld implements FetchDownloadProgressOwner, FetchEnqueue
 			
 			if (ActiveDescription == null)
 			{
-				if (NewDownloadDescription.RequestID == null || NewDownloadDescription.RequestID.isEmpty())
-				{
-					//Don't even have an entry in the old list, so it's completely new
-					QueueNewDownloadDescription(NewDownloadDescription, LastRequest);
-					continue;
-				}
-				
-				ActiveDescription = TempRequestedDownloadsCopy.values().stream().filter(Description -> NewDownloadDescription.RequestID.equals(Description.RequestID)).findFirst().orElse(null);
-				
-				if (ActiveDescription == null)
-				{
-					//Don't even have an entry for this RequestID in the old list, so it's completely new
-					QueueNewDownloadDescription(NewDownloadDescription, LastRequest);
-					continue;
-				}
+				//Don't even have an entry in the old list, so it's completely new
+				QueueNewDownloadDescription(NewDownloadDescription, LastRequest);
+				continue;
 			}
 			
 			//Need to update our DownloadDescription
@@ -479,10 +467,10 @@ public class FetchManagerOld implements FetchDownloadProgressOwner, FetchEnqueue
 			return null;
 		}
 		
-		final int NumURLs = Description.URLs.size();
+		final int NumURLs = Description.URLs.length;
 		final int URLIndexToUse = Description.CurrentRetryCount % NumURLs;
 		
-		return Description.URLs.get(URLIndexToUse);
+		return Description.URLs[URLIndexToUse];
 	}
 	
 	private boolean IsOutOfRetries(DownloadDescription Description)
@@ -895,7 +883,7 @@ public class FetchManagerOld implements FetchDownloadProgressOwner, FetchEnqueue
 			return;
 		}
 
-		MatchingDescription.ProgressListener.OnDownloadEnqueued(DestinationLocation, true);
+		MatchingDescription.ProgressListener.OnDownloadEnqueued(MatchingDescription);
 	}
 
 	@Override
@@ -929,7 +917,9 @@ public class FetchManagerOld implements FetchDownloadProgressOwner, FetchEnqueue
 			return;
 		}
 
-		MatchingDescription.ProgressListener.OnDownloadEnqueued(DestinationLocation, false);
+		RetryDownload(DestinationLocation);
+
+		MatchingDescription.ProgressListener.OnDownloadEnqueued(MatchingDescription);
 	}
 	
 	public synchronized void HandleEnqueueCheckCallback(String DestinationLocation, @Nullable Download FetchDownload)
@@ -979,7 +969,12 @@ public class FetchManagerOld implements FetchDownloadProgressOwner, FetchEnqueue
 			MatchingDescription.CachedFetchID = FetchDownload.getId();
 		}
 
-		MatchingDescription.ProgressListener.OnDownloadEnqueued(DestinationLocation, bIsAlreadyEnqueued);
+		if (!bIsAlreadyEnqueued)
+		{
+			RetryDownload(DestinationLocation);
+		}
+
+		MatchingDescription.ProgressListener.OnDownloadEnqueued(MatchingDescription);
 	}
 
 	public boolean IsErrorPossiblyIndicatingPreviouslyQueued(@NonNull Error EnqueueError)
diff --git a/Engine/Plugins/Online/Android/AndroidFetchBackgroundDownload/Source/AndroidFetchBackgroundDownload/Private/AndroidDownloadManager.cpp b/Engine/Plugins/Online/Android/AndroidFetchBackgroundDownload/Source/AndroidFetchBackgroundDownload/Private/AndroidDownloadManager.cpp
new file mode 100644
--- /dev/null
+++ b/Engine/Plugins/Online/Android/AndroidFetchBackgroundDownload/Source/AndroidFetchBackgroundDownload/Private/AndroidDownloadManager.cpp
@@ -0,0 +1,612 @@
+// Copyright Epic Games, Inc. All Rights Reserved.
+
+#include "AndroidDownloadManager.h"
+
+#include "Misc/ConfigCacheIni.h"
+#include "Misc/CoreDelegates.h"
+
+#include "BackgroundHttpModule.h"
+#include "PlatformBackgroundHttp.h"
+
+#include "Android/AndroidPlatform.h"
+#include "Android/AndroidJavaEnv.h"
+#include "Android/AndroidPlatformFile.h"
+
+#include "Containers/UnrealString.h"
+#include "Misc/Paths.h"
+
+#include "BuildSettings.h"
+#include "Async/Async.h"
+#include "HAL/FileManager.h"
+#include "HAL/PlatformFileManager.h"
+
+namespace UE::Jni::Download::DataStructs
+{
+	struct FDownloadDescription: Java::Lang::FObject
+	{
+		static constexpr FAnsiStringView ClassName = "com/epicgames/unreal/download/datastructs/DownloadDescription";
+		
+		inline static TConstructor<FDownloadDescription(Java::Lang::FString* RequestID, Java::Lang::FString* DestinationLocation, jint MaxRetryCount, jint IndividualURLRetryCount, jint RequestPriority, jint GroupID, Java::Lang::TArray<Java::Lang::FString*>* URLs, jboolean bIsPaused, jlong TotalBytesNeeded)> New;
+		
+		static constexpr FMember Members[]
+		{
+			UE_JNI_MEMBER(New)
+		};
+	};
+}
+
+namespace UE::Jni::Download
+{
+	struct FDownloadManager: Java::Lang::FObject
+	{
+		static constexpr FAnsiStringView ClassName = "com/epicgames/unreal/download/DownloadManager";
+
+		inline static TMember<FDownloadManager, void (Java::Util::FHashMap* Params, jboolean bWaitTillDone)> Initialize;
+		inline static TMember<FDownloadManager, void (jint MaxConcurrentDownloads)> SetMaxConcurrentDownloads;
+		inline static TMember<FDownloadManager, void (DataStructs::FDownloadDescription* DownloadDescription)> Enqueue;
+		inline static TMember<FDownloadManager, void (Java::Lang::FString* DestinationLocation)> Pause;
+		inline static TMember<FDownloadManager, void (Java::Lang::FString* DestinationLocation)> Resume;
+		inline static TMember<FDownloadManager, void (Java::Lang::FString* DestinationLocation)> Cancel;
+		
+		static void JNICALL OnProgress(JNIEnv* env, Java::Lang::TClass<FDownloadManager>* clazz, Java::Lang::FString* DestinationLocation, jlong TotalBytesWritten, jlong BytesWrittenSinceLastCall)
+		{
+			AsyncTask(ENamedThreads::GameThread, [DestinationLocation = Env.Marshal(DestinationLocation), TotalBytesWritten, BytesWrittenSinceLastCall]
+			{
+				if (TSharedPtr<Online::Download::Android::FDownloadManager> Manager = StaticCastSharedPtr<Online::Download::Android::FDownloadManager>(FBackgroundHttpModule::Get().GetBackgroundHttpManager()))
+				{
+					Manager->OnProgress(DestinationLocation, TotalBytesWritten, BytesWrittenSinceLastCall);
+				}
+			});
+		}
+		
+		static void JNICALL OnComplete(JNIEnv* env, Java::Lang::TClass<FDownloadManager>* clazz, Java::Lang::FString* DestinationLocation, jboolean bDidSucceed, jlong TotalBytesDownloaded, jlong DownloadDuration, jlong DownloadStartTimeUTC, jlong DownloadEndTimeUTC)
+		{
+			AsyncTask(ENamedThreads::GameThread, [DestinationLocation = Env.Marshal(DestinationLocation), bDidSucceed, TotalBytesDownloaded, DownloadDuration, DownloadStartTimeUTC, DownloadEndTimeUTC]
+			{
+				if (TSharedPtr<Online::Download::Android::FDownloadManager> Manager = StaticCastSharedPtr<Online::Download::Android::FDownloadManager>(FBackgroundHttpModule::Get().GetBackgroundHttpManager()))
+				{
+					Manager->OnComplete(DestinationLocation, bDidSucceed, TotalBytesDownloaded, DownloadDuration, DownloadStartTimeUTC, DownloadEndTimeUTC);
+				}
+			});
+		}
+		
+		static constexpr FMember Members[]
+		{
+			UE_JNI_MEMBER(Initialize),
+			UE_JNI_MEMBER(SetMaxConcurrentDownloads),
+			UE_JNI_MEMBER(Enqueue),
+			UE_JNI_MEMBER(Pause),
+			UE_JNI_MEMBER(Resume),
+			UE_JNI_MEMBER(Cancel)
+		};
+		
+		static constexpr FNativeMethod NativeMethods[]
+		{
+			UE_JNI_NATIVE_METHOD(OnProgress),
+			UE_JNI_NATIVE_METHOD(OnComplete)
+		};
+	};
+}
+
+#if USE_ANDROID_JNI
+extern void AndroidThunkCpp_SetSharedPreferenceString(const FString& Group, const FString& Key, const FString& Value);
+#endif
+
+namespace UE::Online::Download::Android
+{
+	bool FDownloadManager::FRequest::ProcessRequest()
+	{
+		if (!ensureAlwaysMsgf(IsInGameThread(), TEXT("Should only ever call ProcessRequest from GameThread! Cannot Process!")))
+		{
+			return false;
+		}
+		
+		StaticCastSharedPtr<FDownloadManager>(FBackgroundHttpModule::Get().GetBackgroundHttpManager())->EnqueueRequest(*this);
+		return true;
+	}
+	
+	void FDownloadManager::FRequest::CancelRequest()
+	{
+		if (!ensureAlwaysMsgf(IsInGameThread(), TEXT("Should only ever call CancelRequest from GameThread! Cannot Cancel!")))
+		{
+			return;
+		}
+		
+		if (TSharedPtr<FDownloadManager> Manager = Parent.Pin())
+		{
+			Manager->CancelRequest(*this);
+		}
+	}
+
+	void FDownloadManager::FRequest::PauseRequest()
+	{
+		if (!ensureAlwaysMsgf(IsInGameThread(), TEXT("Should only ever call PauseRequest from GameThread! Cannot Pause!")))
+		{
+			return;
+		}
+		
+		UE_LOG(LogBackgroundHttpManager, Verbose, TEXT("PauseRequest called on %s"), *GetRequestID());
+
+		bIsPaused = true;
+		
+		if (Parent.IsValid())
+		{
+			Jni::Download::FDownloadManager::Pause(Jni::Class<Jni::Download::FDownloadManager>, *Jni::Env.Marshal(DestinationLocation));
+		}
+	}
+
+	void FDownloadManager::FRequest::ResumeRequest()
+	{
+		if (!ensureAlwaysMsgf(IsInGameThread(), TEXT("Should only ever call ResumeRequest from GameThread! Cannot Resume!")))
+		{
+			return;
+		}
+		
+		UE_LOG(LogBackgroundHttpManager, Verbose, TEXT("ResumeRequest called on %s"), *GetRequestID());
+
+		bIsPaused = false;
+
+		if (Parent.IsValid())
+		{
+			Jni::Download::FDownloadManager::Resume(Jni::Class<Jni::Download::FDownloadManager>, *Jni::Env.Marshal(DestinationLocation));
+		}
+	}
+
+	int FDownloadManager::FRequest::GetPriorityAsAndroidPriority() const
+	{
+		switch (RequestPriority)
+		{
+		case EBackgroundHTTPPriority::High:
+			return 1;
+		case EBackgroundHTTPPriority::Low:
+			return -1;
+		case EBackgroundHTTPPriority::Normal:
+			return 0;
+		default:
+			ensureAlwaysMsgf(false, TEXT("Missing EBackgroundHTTPPriority in GetPriorityAsAndroidPriority!"));
+			return 0;
+		}
+	}
+
+	FDownloadManager::FDownloadManager()
+		: MaxActiveDownloads{4}
+		, FileHashHelper{MakeShared<FBackgroundHttpFileHashHelper, ESPMode::ThreadSafe>()}
+	{
+		Jni::Env.Initialize<Jni::Download::DataStructs::FDownloadDescription>();
+		Jni::Env.Initialize<Jni::Download::FDownloadManager>();
+	}
+
+	void FDownloadManager::Initialize(bool bClearPreExistingRequestsAtStartup)
+	{
+		AndroidThunkCpp_SetSharedPreferenceString(TEXT("BackgroundPreferences"), TEXT("TempDownloadRootPath"), FBackgroundHttpFileHashHelper().GetTemporaryRootPath());
+
+		ensureAlwaysMsgf(GConfig->GetInt(TEXT("BackgroundHttp"), TEXT("MaxActiveDownloads"), MaxActiveDownloads, GEngineIni), TEXT("No value found for MaxActiveDownloads! Defaulting to 4!"));
+		
+		AndroidBackgroundHTTPManagerDefaultLocalizedText.InitFromIniSettings("AndroidFetchBackgroundDownload");
+		
+		{
+			using namespace Jni;
+			
+			FScopedJavaObject<Java::Util::FHashMap*> Params = Java::Util::FHashMap::New(Class<Java::Util::FHashMap>);
+
+			//Set our MaxActiveDownloads in the underlying java layer to match our expectation
+			Java::Util::FHashMap::put(*Params, *Env.Marshal(FAndroidNativeDownloadWorkerParameterKeys::DOWNLOAD_MAX_CONCURRENT_REQUESTS_KEY), *Java::Lang::FInteger::New(Class<Java::Lang::FInteger>, MaxActiveDownloads));
+
+			//Make sure we pass in localized notification text bits for the important worker keys
+			Java::Util::FHashMap::put(*Params, *Env.Marshal(FAndroidNativeDownloadWorkerParameterKeys::NOTIFICATION_CONTENT_TITLE_KEY), *Env.Marshal(AndroidBackgroundHTTPManagerDefaultLocalizedText.DefaultNotificationText_Title.GetText().ToString()));
+			Java::Util::FHashMap::put(*Params, *Env.Marshal(FAndroidNativeDownloadWorkerParameterKeys::NOTIFICATION_CONTENT_COMPLETE_TEXT_KEY), *Env.Marshal(AndroidBackgroundHTTPManagerDefaultLocalizedText.DefaultNotificationText_Complete.GetText().ToString()));
+			Java::Util::FHashMap::put(*Params, *Env.Marshal(FAndroidNativeDownloadWorkerParameterKeys::NOTIFICATION_CONTENT_CANCEL_DOWNLOAD_TEXT_KEY), *Env.Marshal(AndroidBackgroundHTTPManagerDefaultLocalizedText.DefaultNotificationText_Cancel.GetText().ToString()));
+			Java::Util::FHashMap::put(*Params, *Env.Marshal(FAndroidNativeDownloadWorkerParameterKeys::NOTIFICATION_CONTENT_NO_INTERNET_TEXT_KEY), *Env.Marshal(AndroidBackgroundHTTPManagerDefaultLocalizedText.DefaultNotificationText_NoInternet.GetText().ToString()));
+			Java::Util::FHashMap::put(*Params, *Env.Marshal(FAndroidNativeDownloadWorkerParameterKeys::NOTIFICATION_CONTENT_AIRPLANE_MODE_TEXT_KEY), *Env.Marshal(AndroidBackgroundHTTPManagerDefaultLocalizedText.DefaultNotificationText_AirplaneMode.GetText().ToString()));
+			Java::Util::FHashMap::put(*Params, *Env.Marshal(FAndroidNativeDownloadWorkerParameterKeys::NOTIFICATION_CONTENT_DATA_SAVER_ENABLED_TEXT_KEY), *Env.Marshal(AndroidBackgroundHTTPManagerDefaultLocalizedText.DefaultNotificationText_DataSaverEnabled.GetText().ToString()));
+			Java::Util::FHashMap::put(*Params, *Env.Marshal(FAndroidNativeDownloadWorkerParameterKeys::NOTIFICATION_CONTENT_WAITING_FOR_CELLULAR_TEXT_KEY), *Env.Marshal(AndroidBackgroundHTTPManagerDefaultLocalizedText.DefaultNotificationText_WaitingForCellular.GetText().ToString()));
+			Java::Util::FHashMap::put(*Params, *Env.Marshal(FAndroidNativeDownloadWorkerParameterKeys::NOTIFICATION_CONTENT_APPROVE_TEXT_KEY), *Env.Marshal(AndroidBackgroundHTTPManagerDefaultLocalizedText.DefaultNotificationText_Approve.GetText().ToString()));
+
+			//Expect our ContentText to have a {DownloadPercent} argument in it by default, so this will replace that with the Java string format argument so Java can insert the appropriate value
+			FFormatNamedArguments Arguments;
+			Arguments.Emplace(TEXT("DownloadPercent"), FText::FromString(TEXT("%3d%%")));
+			FText UpdatedContentText = FText::Format(AndroidBackgroundHTTPManagerDefaultLocalizedText.DefaultNotificationText_Content.GetText(), Arguments);
+			Java::Util::FHashMap::put(*Params, *Env.Marshal(FAndroidNativeDownloadWorkerParameterKeys::NOTIFICATION_CONTENT_TEXT_KEY), *Env.Marshal(UpdatedContentText.ToString()));
+
+			Jni::Download::FDownloadManager::Initialize(Class<Jni::Download::FDownloadManager>, *Params, bClearPreExistingRequestsAtStartup);
+		}
+		
+		FileHashHelper->LoadData();
+		
+		if (bClearPreExistingRequestsAtStartup)
+		{
+			DeleteAllTemporaryFiles();
+		}		
+		
+		DeleteStaleTempFiles();
+
+		// we download to our temporary root path, and then in a subdirectory named after our branch name;
+		// this way, the fetch managers can clean up stuff they don't know (e.g. downloaded temp files from older versions)
+		if (IAndroidPlatformFile::GetPlatformPhysical().CreateDirectory(*FBackgroundHttpFileHashHelper::GetTemporaryRootPath()))
+		{
+			const FString BuildDownloadPath = FPaths::Combine(FBackgroundHttpFileHashHelper::GetTemporaryRootPath(), BuildSettings::GetBranchName());
+			if (!IAndroidPlatformFile::GetPlatformPhysical().CreateDirectory(*BuildDownloadPath))
+			{
+				UE_LOG(LogBackgroundHttpManager, Display, TEXT("Failed to create build temp download path: %s"), *BuildDownloadPath);
+			}
+		}
+	}
+
+	void FDownloadManager::Shutdown()
+	{
+		ensureAlwaysMsgf(IsInGameThread(), TEXT("Called from un-expected thread! Potential error in an implementation of background downloads!"));
+		
+		Requests.Empty();
+	}
+
+	void FDownloadManager::AddRequest(const FBackgroundHttpRequestPtr Request)
+	{
+		ensureAlwaysMsgf(IsInGameThread(), TEXT("Called from un-expected thread! Potential error in an implementation of background downloads!"));
+		
+		EnqueueRequest(static_cast<FRequest&>(*Request));
+	}
+
+	void FDownloadManager::RemoveRequest(const FBackgroundHttpRequestPtr Request)
+	{
+		ensureAlwaysMsgf(IsInGameThread(), TEXT("Called from un-expected thread! Potential error in an implementation of background downloads!"));
+		
+		CancelRequest(static_cast<FRequest&>(*Request));		
+	}
+
+	void FDownloadManager::DeleteAllTemporaryFiles()
+	{
+		UE_LOG(LogBackgroundHttpManager, Log, TEXT("Cleaning Up Temporary Files"));
+
+		TArray<FString> FilesToDelete;
+
+		//Default implementation is to just delete everything in the Root Folder non-recursively.
+		IPlatformFile& PlatformFile = FPlatformFileManager::Get().GetPlatformFile();
+		PlatformFile.FindFiles(FilesToDelete, *FBackgroundHttpFileHashHelper::GetTemporaryDownloadPath(), *FBackgroundHttpFileHashHelper::GetTempFileExtension());
+
+		for (const FString& File : FilesToDelete)
+		{
+			UE_LOG(LogBackgroundHttpManager, Log, TEXT("Deleting File:%s"), *File);
+			const bool bDidDelete = PlatformFile.DeleteFile(*File);
+
+			if (!bDidDelete)
+			{
+				UE_LOG(LogBackgroundHttpManager, Warning, TEXT("Failure to Delete Temp File:%s"), *File);
+			}
+		}
+	}
+
+	void FDownloadManager::SetCellularPreference(int32 Value)
+	{
+		FPlatformMisc::SetCellularPreference(Value);
+	}
+
+	void FDownloadManager::SetLimitedDataPreference(int32 Value)
+	{
+	}
+
+	void FDownloadManager::SetMaxActiveDownloads(int32 Value)
+	{
+		check(Value > 0);
+		if (MaxActiveDownloads != Value)
+		{
+			MaxActiveDownloads = Value;
+			
+			Jni::Download::FDownloadManager::SetMaxConcurrentDownloads(Jni::Class<Jni::Download::FDownloadManager>, MaxActiveDownloads);
+		}
+	}
+
+	void FDownloadManager::CleanUpDataAfterCompletingRequest(const FBackgroundHttpRequestPtr Request)
+	{
+		//Need to free up all these URL's hashes in FileHashHelper so that future URLs can use those temp files
+		for (const FString& URL : Request->GetURLList())
+		{
+			FileHashHelper->RemoveURLMapping(URL, Request->GetRequestID());
+		}
+		
+		FileHashHelper->SaveData();
+	}
+	
+	void FDownloadManager::EnqueueRequest(FRequest& _Request)
+	{
+		TSharedPtr<FRequest> Request = SharedThis(&_Request);
+		
+		UE_LOG(LogBackgroundHttpRequest, Verbose, TEXT("Processing Request - RequestID:%s"), *Request->RequestID);
+		
+		if (!ensure(!Request->URLList.IsEmpty()))
+		{
+			Request->CompleteWithExistingResponseData(FPlatformBackgroundHttp::ConstructBackgroundResponse(EHttpResponseCodes::Unknown, {}));
+			return;
+		}
+		
+		for (const FString& URL : Request->GetURLList())
+		{
+			if (FileHashHelper->FindTempFilenameMappingForURL(URL, Request->GetRequestID(), Request->DestinationLocation))
+			{
+				Request->DestinationLocation = FBackgroundHttpFileHashHelper::GetFullPathOfTempFilename(Request->DestinationLocation);
+				
+				if (IPlatformFile& PlatformFile = FPlatformFileManager::Get().GetPlatformFile(); PlatformFile.FileExists(*Request->DestinationLocation))
+				{
+					int64 ExistingFileSize = PlatformFile.FileSize(*Request->DestinationLocation);
+					UE_LOG(LogBackgroundHttpManager, Display, TEXT("Found existing background task to associate with! RequestID:%s | ExistingFileSize:%lld | ExistingFilePath:%s"), *Request->GetRequestID(), ExistingFileSize, *Request->DestinationLocation);
+
+					Request->OnProgressUpdated().ExecuteIfBound(Request, ExistingFileSize, ExistingFileSize);
+					Request->SetMetricsExtended(
+					{
+						.TotalBytesDownloaded = int32(ExistingFileSize),
+						.DownloadDuration = 0.0f,
+						.FetchStartTimeUTC = FDateTime::FromUnixTimestamp(0),
+						.FetchEndTimeUTC = FDateTime::FromUnixTimestamp(0)
+					});
+					Request->CompleteWithExistingResponseData(FPlatformBackgroundHttp::ConstructBackgroundResponse(EHttpResponseCodes::Ok, Request->DestinationLocation));
+					return;
+				}
+				
+				break;
+			}
+		}
+		
+		if (Request->DestinationLocation.IsEmpty())
+		{
+			Request->DestinationLocation = FBackgroundHttpFileHashHelper::GetFullPathOfTempFilename(FileHashHelper->FindOrAddTempFilenameMappingForURL(Request->URLList[0], Request->RequestID));
+			FileHashHelper->SaveData();
+		}
+		
+		if (!ensure(Requests.FindOrAdd(Request->DestinationLocation, Request) == Request))
+		{
+			Request->CompleteWithExistingResponseData(FPlatformBackgroundHttp::ConstructBackgroundResponse(EHttpResponseCodes::Unknown, {}));
+			return;
+		}
+		
+		Request->Parent = StaticCastWeakPtr<FDownloadManager>(AsWeak());
+		
+		static constexpr int DefaultIndividualURLRetryCount = 3;
+		
+		Jni::Download::FDownloadManager::Enqueue(Jni::Class<Jni::Download::FDownloadManager>, *Jni::Download::DataStructs::FDownloadDescription::New(Jni::Class<Jni::Download::DataStructs::FDownloadDescription>, *Jni::Env.Marshal(Request->RequestID), *Jni::Env.Marshal(Request->DestinationLocation), Request->NumberOfTotalRetries, DefaultIndividualURLRetryCount, Request->GetPriorityAsAndroidPriority(), 0, *Jni::Env.Marshal(Request->GetURLList()), Request->bIsPaused, Request->GetExpectedResultSize()));
+	}
+
+	void FDownloadManager::CancelRequest(FRequest& _Request)
+	{
+		if (TSharedPtr<FRequest> Request; Requests.RemoveAndCopyValue(_Request.DestinationLocation, Request))
+		{
+			UE_LOG(LogBackgroundHttpManager, Verbose, TEXT("CancelRequest called on %s"), *Request->RequestID);
+			
+			Request->Parent = nullptr;
+		
+			Jni::Download::FDownloadManager::Cancel(Jni::Class<Jni::Download::FDownloadManager>, *Jni::Env.Marshal(Request->DestinationLocation));
+		}
+	}
+	
+	void FDownloadManager::OnProgress(const FString& DestinationLocation, int64 TotalBytesWritten, int64 BytesWrittenSinceLastCall) const
+	{
+		UE_LOG(LogBackgroundHttpManager, VeryVerbose, TEXT("Download Progress... DestinationLocation:%s BytesWrittenSinceLastCall:%lld TotalBytesWritten:%lld"), *DestinationLocation, BytesWrittenSinceLastCall, TotalBytesWritten);
+		
+		if (TSharedPtr<FRequest> Request = Requests.FindRef(DestinationLocation))
+		{
+			Request->OnProgressUpdated().ExecuteIfBound(Request, TotalBytesWritten, BytesWrittenSinceLastCall);
+		}
+	}
+	
+	void FDownloadManager::OnComplete(const FString& DestinationLocation, bool bDidSucceed, int64 TotalBytesDownloaded, int64 DownloadDuration, int64 DownloadStartTimeUTC, int64 DownloadEndTimeUTC)
+	{
+		UE_LOG(LogBackgroundHttpManager, Log, TEXT("DownloadMetrics... DestinationLocation:%s TotalBytesDownloaded:%ld DownloadDuration:%ld"), *DestinationLocation, TotalBytesDownloaded, DownloadDuration);
+		UE_LOG(LogBackgroundHttpManager, Log, TEXT("DownloadComplete... DestinationLocation:%s bWasSuccess:%d"), *DestinationLocation, int(bDidSucceed));
+		
+		if (TSharedPtr<FRequest> Request; Requests.RemoveAndCopyValue(DestinationLocation, Request))
+		{
+			Request->Parent = nullptr;
+				
+			EHttpResponseCodes::Type ResponseCodeToUse = bDidSucceed && FPlatformFileManager::Get().GetPlatformFile().FileExists(*Request->DestinationLocation) ? EHttpResponseCodes::Ok : EHttpResponseCodes::Unknown;
+			
+			Request->NotifyRequestMetricsExtendedAvailable(
+			{
+				.TotalBytesDownloaded = (int32)TotalBytesDownloaded,
+				.DownloadDuration = (float)((double)DownloadDuration / 1000.0f),
+				.FetchStartTimeUTC = FDateTime::FromUnixTimestampDecimal((double)DownloadStartTimeUTC / 1000.0f),
+				.FetchEndTimeUTC = FDateTime::FromUnixTimestampDecimal((double)DownloadEndTimeUTC / 1000.0f)
+			});
+			Request->CompleteWithExistingResponseData(FPlatformBackgroundHttp::ConstructBackgroundResponse(ResponseCodeToUse, Request->DestinationLocation));
+		}
+		else
+		{
+			UE_LOG(LogBackgroundHttpManager, Log, TEXT("Taking no action as DestinationLocation:%s did not have a corresponding active Request"), *DestinationLocation);
+		}
+	}
+	
+	void FDownloadManager::DeleteStaleTempFiles() const
+	{
+		//Parse our .ini values to determine how much we clean in this function
+		double FileAgeTimeOutSettings = -1;
+		bool bDeleteTimedOutFiles;
+		bool bDeleteTempFilesWithoutURLMappings = false;
+		bool bRemoveURLMappingEntriesWithoutPhysicalTempFiles = false;
+		{
+			GConfig->GetDouble(TEXT("BackgroundHttp"), TEXT("TempFileTimeOutSeconds"), FileAgeTimeOutSettings, GEngineIni);
+			bDeleteTimedOutFiles = (FileAgeTimeOutSettings >= 0);
+			
+			GConfig->GetBool(TEXT("BackgroundHttp"), TEXT("DeleteTempFilesWithoutURLMappingEntries"), bDeleteTempFilesWithoutURLMappings, GEngineIni);
+			
+			GConfig->GetBool(TEXT("BackgroundHttp"), TEXT("RemoveURLMappingEntriesWithoutPhysicalTempFiles"), bRemoveURLMappingEntriesWithoutPhysicalTempFiles, GEngineIni);
+			
+			UE_LOG(LogBackgroundHttpManager, Log, TEXT("Stale Settings -- TempFileTimeOutSeconds:%f DeleteTempFilesWithoutURLMappingEntries:%d RemoveURLMappingEntriesWithoutPhysicalTempFiles:%d"), static_cast<float>(FileAgeTimeOutSettings),static_cast<int>(bDeleteTempFilesWithoutURLMappings),static_cast<int>(bRemoveURLMappingEntriesWithoutPhysicalTempFiles));
+		}
+		const bool bWillDoAnyWork = bDeleteTimedOutFiles || bDeleteTempFilesWithoutURLMappings || bRemoveURLMappingEntriesWithoutPhysicalTempFiles;
+		
+		//Only bother gathering temp files if we will actually be doing something with them
+		TArray<FString> AllTempFilesToCheck;
+		if (bWillDoAnyWork)
+		{
+			GatherAllTempFilenames(AllTempFilesToCheck);
+			UE_LOG(LogBackgroundHttpManager, Display, TEXT("Found %d temp download files."), AllTempFilesToCheck.Num());
+		}
+		
+		//Handle all timed out files based on the .ini time out settings
+		//can be turned off by setting
+		if (bDeleteTimedOutFiles)
+		{
+			TArray<FString> TimedOutFiles;
+			GatherTempFilesOlderThen(TimedOutFiles, FileAgeTimeOutSettings, &AllTempFilesToCheck);
+		
+			TArray<FString> TimeOutDeleteFullPaths;
+			ConvertAllTempFilenamesToFullPaths(TimeOutDeleteFullPaths, TimedOutFiles);
+			
+			for (const FString& FullFilePath : TimeOutDeleteFullPaths)
+			{
+				if (IFileManager::Get().Delete(*FullFilePath))
+				{
+					UE_LOG(LogBackgroundHttpManager, Log, TEXT("Successfully deleted %s due to time out settings"), *FullFilePath);
+				}
+				else
+				{
+					UE_LOG(LogBackgroundHttpManager, Error, TEXT("Failed to delete timed out file %s"), *FullFilePath);
+				}
+			}
+			
+			//Should remove these files from the list of files we are checking as we know they are already invalid from timing out, so we shouldn't check them twice
+			for(const FString& RemovedFile : TimedOutFiles)
+			{
+				AllTempFilesToCheck.Remove(RemovedFile);
+			}
+		}
+		
+		//Handle all temp files that should be removed because they are missing a corresponding URL mapping
+		if (bDeleteTempFilesWithoutURLMappings)
+		{
+			TArray<FString> MissingURLMappingFiles;
+			GatherTempFilesWithoutURLMappings(MissingURLMappingFiles, &AllTempFilesToCheck);
+			
+			TArray<FString> MissingURLDeleteFullPaths;
+			ConvertAllTempFilenamesToFullPaths(MissingURLDeleteFullPaths, MissingURLMappingFiles);
+			
+			for (const FString& FullFilePath : MissingURLDeleteFullPaths)
+			{
+				if (IFileManager::Get().Delete(*FullFilePath))
+				{
+					UE_LOG(LogBackgroundHttpManager, Log, TEXT("Successfully deleted %s due to missing a URL mapping for this temp data"), *FullFilePath);
+				}
+				else
+				{
+					UE_LOG(LogBackgroundHttpManager, Error, TEXT("Failed to delete file %s that was being deleted due to a missing URL mapping"), *FullFilePath);
+				}
+			}
+			
+			//Should remove these files from the list of files we are checking as we know they are already invalid from timing out, so we shouldn't check them twice
+			for(const FString& RemovedFile : MissingURLMappingFiles)
+			{
+				AllTempFilesToCheck.Remove(RemovedFile);
+			}
+		}
+		
+		if (bRemoveURLMappingEntriesWithoutPhysicalTempFiles)
+		{
+			//Remove all URL map entries that don't correspond to a physical file on disk
+			FileHashHelper->DeleteURLMappingsWithoutTempFiles();
+			FileHashHelper->SaveData();
+		}
+		
+		UE_LOG(LogBackgroundHttpManager, Log, TEXT("Kept %d temp download files:"), AllTempFilesToCheck.Num());
+		for (const FString& ValidFile : AllTempFilesToCheck)
+		{
+			UE_LOG(LogBackgroundHttpManager, Verbose, TEXT("Kept: %s"), *ValidFile);
+		}
+	}
+
+	void FDownloadManager::GatherTempFilesOlderThen(TArray<FString>& OutTimedOutTempFilenames,double SecondsToConsiderOld, const TArray<FString>* OptionalFileList /* = nullptr */)
+	{
+		OutTimedOutTempFilenames.Empty();
+		
+		TArray<FString> GatheredFullFilePathFiles;
+		
+		//OptionalFileList was not supplied so we need to gather all temp files to check as full file paths
+		if (nullptr == OptionalFileList)
+		{
+			GatherAllTempFilenames(GatheredFullFilePathFiles, true);
+		}
+		//We supplied an OptionalFileList, but we still need a full file path list for this operation
+		else
+		{
+			ConvertAllTempFilenamesToFullPaths(GatheredFullFilePathFiles, *OptionalFileList);
+		}
+
+		if (SecondsToConsiderOld >= 0)
+		{
+			UE_LOG(LogBackgroundHttpManager, Verbose, TEXT("Checking for BackgroundHTTP temp files that are older then: %lf"), SecondsToConsiderOld);
+			
+			for (const FString& FullFilePath : GatheredFullFilePathFiles)
+			{
+				FFileStatData FileData = IFileManager::Get().GetStatData(*FullFilePath);
+				FTimespan TimeSinceCreate = FDateTime::UtcNow() - FileData.CreationTime;
+
+				const double FileAge = TimeSinceCreate.GetTotalSeconds();
+				const bool bShouldReturn = (FileAge > SecondsToConsiderOld);
+				if (bShouldReturn)
+				{
+					UE_LOG(LogBackgroundHttpManager, Verbose, TEXT("FoundTempFile: %s with age %lf"), *FullFilePath, FileAge);
+					
+					//Need to save output as just filename to be consistent with other functions
+					OutTimedOutTempFilenames.Add(FPaths::GetCleanFilename(FullFilePath));
+				}
+			}
+		}
+	}
+
+	void FDownloadManager::GatherTempFilesWithoutURLMappings(TArray<FString>& OutTempFilesMissingURLMappings, TArray<FString>* OptionalFileList /*= nullptr */) const
+	{
+		OutTempFilesMissingURLMappings.Empty();
+
+		TArray<FString>* FileListToCheckPtr = OptionalFileList;
+
+		//OptionalFileList was not supplied so we need to gather all temp files to check
+		TArray<FString> GatheredFiles;
+		if (nullptr == FileListToCheckPtr)
+		{
+			GatherAllTempFilenames(GatheredFiles, false);
+			FileListToCheckPtr = &GatheredFiles;
+		}
+		
+		TArray<FString>& FilesToCheckRef = *FileListToCheckPtr;
+		for (const FString& File : FilesToCheckRef)
+		{
+			FString FoundURL = {};
+			if (!FileHashHelper->FindMappedURLForTempFilename(File, FoundURL))
+			{
+				OutTempFilesMissingURLMappings.Add(File);
+			}
+		}
+	}
+
+	void FDownloadManager::GatherAllTempFilenames(TArray<FString>& OutAllTempFilenames, bool bOutputAsFullPaths /* = false */)
+	{
+		OutAllTempFilenames.Empty();
+		
+		const FString& DirectoryToCheck = FBackgroundHttpFileHashHelper::GetTemporaryDownloadPath();
+		
+		TArray<FString> AllFilenames;
+		IFileManager::Get().FindFiles(AllFilenames, *DirectoryToCheck, *FBackgroundHttpFileHashHelper::GetTempFileExtension());
+		
+		//Make into full paths for output
+		for (const FString& Filename : AllFilenames)
+		{
+			if (bOutputAsFullPaths)
+			{
+				OutAllTempFilenames.Add(FPaths::Combine(DirectoryToCheck, Filename));
+			}
+			else
+			{
+				OutAllTempFilenames.Add(Filename);
+			}
+		}
+	}
+
+	void FDownloadManager::ConvertAllTempFilenamesToFullPaths(TArray<FString>& OutFilenamesAsFullPaths, const TArray<FString>& FilenamesToConvertToFullPaths)
+	{
+		//Store this separately so we don't get bad behavior if the same Array is supplied for both parameters
+		TArray<FString> FilenamesToOutput;
+		
+		for(const FString& ExistingFilename : FilenamesToConvertToFullPaths)
+		{
+			FilenamesToOutput.Add(FBackgroundHttpFileHashHelper::GetFullPathOfTempFilename(ExistingFilename));
+		}
+		
+		OutFilenamesAsFullPaths = FilenamesToOutput;
+	}
+}
diff --git a/Engine/Plugins/Online/Android/AndroidFetchBackgroundDownload/Source/AndroidFetchBackgroundDownload/Private/AndroidFetchPlatformBackgroundHttp.cpp b/Engine/Plugins/Online/Android/AndroidFetchBackgroundDownload/Source/AndroidFetchBackgroundDownload/Private/AndroidFetchPlatformBackgroundHttp.cpp
--- a/Engine/Plugins/Online/Android/AndroidFetchBackgroundDownload/Source/AndroidFetchBackgroundDownload/Private/AndroidFetchPlatformBackgroundHttp.cpp
+++ b/Engine/Plugins/Online/Android/AndroidFetchBackgroundDownload/Source/AndroidFetchBackgroundDownload/Private/AndroidFetchPlatformBackgroundHttp.cpp
@@ -2,10 +2,21 @@
 
 #include "AndroidFetchPlatformBackgroundHttp.h"
 
+#include "AndroidDownloadManager.h"
 #include "AndroidPlatformBackgroundHttpManager.h"
 #include "AndroidPlatformBackgroundHttpRequest.h"
 #include "AndroidPlatformBackgroundHttpResponse.h"
 
+#if USE_ANDROID_JNI
+extern bool AndroidThunkCpp_GetSharedPreferenceBool(const FString& Group, const FString& Key, bool DefaultValue);
+#endif
+
+static bool UseNewDownloadManager()
+{
+	static bool Value = AndroidThunkCpp_GetSharedPreferenceBool("BackgroundPreferences", "UseNewDownloadManager", false);
+	return Value;
+}
+
 void FAndroidFetchPlatformBackgroundHttp::Initialize()
 {
 }
@@ -16,16 +27,31 @@ void FAndroidFetchPlatformBackgroundHttp::Shutdown()
 
 FBackgroundHttpManagerPtr FAndroidFetchPlatformBackgroundHttp::CreatePlatformBackgroundHttpManager()
 {
-    return MakeShared<FAndroidPlatformBackgroundHttpManager, ESPMode::ThreadSafe>();
+	if (UseNewDownloadManager())
+	{
+		return MakeShared<UE::Online::Download::Android::FDownloadManager, ESPMode::ThreadSafe>();
+	}
+
+	return MakeShared<FAndroidPlatformBackgroundHttpManager, ESPMode::ThreadSafe>();
 }
 
 FBackgroundHttpRequestPtr FAndroidFetchPlatformBackgroundHttp::ConstructBackgroundRequest()
 {
+	if (UseNewDownloadManager())
+	{
+		return MakeShared<UE::Online::Download::Android::FDownloadManager::FRequest, ESPMode::ThreadSafe>();
+	}
+	
 	return MakeShared<FAndroidPlatformBackgroundHttpRequest, ESPMode::ThreadSafe>();
 }
 
 FBackgroundHttpResponsePtr FAndroidFetchPlatformBackgroundHttp::ConstructBackgroundResponse(int32 ResponseCode, const FString& TempFilePath)
 {
+	if (UseNewDownloadManager())
+	{
+		return MakeShared<FBackgroundHttpResponseImpl, ESPMode::ThreadSafe>(TempFilePath, ResponseCode);
+	}
+	
 	return MakeShared<FAndroidPlatformBackgroundHttpResponse, ESPMode::ThreadSafe>(ResponseCode, TempFilePath);
 }
 
diff --git a/Engine/Plugins/Online/Android/AndroidFetchBackgroundDownload/Source/AndroidFetchBackgroundDownload/Public/AndroidDownloadManager.h b/Engine/Plugins/Online/Android/AndroidFetchBackgroundDownload/Source/AndroidFetchBackgroundDownload/Public/AndroidDownloadManager.h
new file mode 100644
--- /dev/null
+++ b/Engine/Plugins/Online/Android/AndroidFetchBackgroundDownload/Source/AndroidFetchBackgroundDownload/Public/AndroidDownloadManager.h
@@ -0,0 +1,86 @@
+// Copyright Epic Games, Inc. All Rights Reserved.
+
+#pragma once
+
+#include "CoreMinimal.h"
+
+#include "AndroidPlatformBackgroundHttpManager.h"
+
+namespace UE::Jni::Download
+{
+	struct FDownloadManager;
+}
+
+namespace UE::Online::Download::Android
+{
+	struct FDownloadManager: IBackgroundHttpManager
+	{
+		struct FRequest: FBackgroundHttpRequestImpl
+		{
+			virtual bool ProcessRequest() override;
+			virtual void CancelRequest() override;
+			virtual void PauseRequest() override;
+			virtual void ResumeRequest() override;
+
+		private:
+			int GetPriorityAsAndroidPriority() const;
+			
+			FString DestinationLocation;
+			
+			TWeakPtr<FDownloadManager> Parent;
+
+			bool bIsPaused = false;
+
+			friend FDownloadManager;
+		};
+
+		FDownloadManager();
+		
+		virtual void Initialize(bool bClearPreExistingRequestsAtStartup) override;
+		virtual void Shutdown() override;
+		
+		virtual void AddRequest(const FBackgroundHttpRequestPtr Request) override;
+		virtual void RemoveRequest(const FBackgroundHttpRequestPtr Request) override;
+		
+		virtual void DeleteAllTemporaryFiles() override;
+		virtual bool IsGenericImplementation() const override { return false; }
+		
+		virtual void SetCellularPreference(int32 Value) override;
+		virtual void SetLimitedDataPreference(int32 Value) override;
+		
+		virtual int GetMaxActiveDownloads() const override { return MaxActiveDownloads; }
+		virtual void SetMaxActiveDownloads(int32 MaxActiveDownloads) override;
+
+		virtual FString GetTempFileLocationForURL(const FString& URL, const FString& RequestID) override { return {}; }
+		
+		virtual void CleanUpDataAfterCompletingRequest(const FBackgroundHttpRequestPtr Request) override;
+		
+	private:
+		virtual bool AssociateWithAnyExistingRequest(const FBackgroundHttpRequestPtr Request) override { return false; }
+		
+	private:
+		void EnqueueRequest(FRequest& Request);
+		void CancelRequest(FRequest& Request);
+		
+		void OnProgress(const FString& DestinationLocation, int64 TotalBytesWritten, int64 BytesWrittenSinceLastCall) const;
+		void OnComplete(const FString& DestinationLocation, bool bDidSucceed, int64 TotalBytesDownloaded, int64 DownloadDuration, int64 DownloadStartTimeUTC, int64 DownloadEndTimeUTC);
+
+	private:
+		void DeleteStaleTempFiles() const;
+		static void GatherTempFilesOlderThen(TArray<FString>& OutTimedOutTempFilenames,double SecondsToConsiderOld, const TArray<FString>* OptionalFileList = nullptr);
+		void GatherTempFilesWithoutURLMappings(TArray<FString>& OutTempFilesMissingURLMappings, TArray<FString>* OptionalFileList = nullptr) const;
+		static void GatherAllTempFilenames(TArray<FString>& OutAllTempFilenames, bool bOutputAsFullPaths = false);
+		static void ConvertAllTempFilenamesToFullPaths(TArray<FString>& OutFilenamesAsFullPaths, const TArray<FString>& FilenamesToConvertToFullPaths);
+		
+	private:
+		int MaxActiveDownloads;
+		
+		TMap<FStringView, TSharedPtr<FRequest>> Requests;
+
+		BackgroundHttpFileHashHelperRef FileHashHelper;
+		
+		FAndroidPlatformBackgroundHttpManager::FAndroidBackgroundHTTPManagerDefaultLocalizedText AndroidBackgroundHTTPManagerDefaultLocalizedText;
+		
+		friend Jni::Download::FDownloadManager;
+	};
+}
\ No newline at end of file
diff --git a/Engine/Plugins/Online/Android/AndroidFetchBackgroundDownload/Source/AndroidFetchBackgroundDownload/Public/AndroidPlatformBackgroundHttpManager.h b/Engine/Plugins/Online/Android/AndroidFetchBackgroundDownload/Source/AndroidFetchBackgroundDownload/Public/AndroidPlatformBackgroundHttpManager.h
--- a/Engine/Plugins/Online/Android/AndroidFetchBackgroundDownload/Source/AndroidFetchBackgroundDownload/Public/AndroidPlatformBackgroundHttpManager.h
+++ b/Engine/Plugins/Online/Android/AndroidFetchBackgroundDownload/Source/AndroidFetchBackgroundDownload/Public/AndroidPlatformBackgroundHttpManager.h
@@ -23,6 +23,11 @@
 //Forward Declarations
 struct FPolyglotTextData;
 
+namespace UE::Online::Download::Android
+{
+	struct FDownloadManager;
+}
+
 /**
  * Manages Background Http request that are currently being processed if we are on an Android Platform
  */
@@ -223,6 +228,8 @@ protected:
 	//We only want to respond to certain broadcast java-side work if we have scheduled BGWork
 	//otherwise it could be stale work from a previous executables run that would look like an error state (Requests we haven't queued yet, etc)
 	static volatile int32 bHasManagerScheduledBGWork;
+	
+	friend UE::Online::Download::Android::FDownloadManager;
 };
 
 //WARNING: These values MUST stay in sync with their values in DownloadWorkerParameterKeys.java!
diff --git a/Engine/Source/Runtime/Online/BackgroundHTTP/Private/BackgroundHttpResponseImpl.cpp b/Engine/Source/Runtime/Online/BackgroundHTTP/Private/BackgroundHttpResponseImpl.cpp
--- a/Engine/Source/Runtime/Online/BackgroundHTTP/Private/BackgroundHttpResponseImpl.cpp
+++ b/Engine/Source/Runtime/Online/BackgroundHTTP/Private/BackgroundHttpResponseImpl.cpp
@@ -8,8 +8,13 @@
 #include "Misc/Paths.h"
 
 FBackgroundHttpResponseImpl::FBackgroundHttpResponseImpl()
-	: TempContentFilePath()
-	, ResponseCode(EHttpResponseCodes::Unknown)
+	: ResponseCode(EHttpResponseCodes::Unknown)
+{
+}
+
+FBackgroundHttpResponseImpl::FBackgroundHttpResponseImpl(FString InTempContentFilePath, int32 InResponseCode)
+	: TempContentFilePath{MoveTemp(InTempContentFilePath)}
+	, ResponseCode{InResponseCode}
 {
 }
 
diff --git a/Engine/Source/Runtime/Online/BackgroundHTTP/Public/BackgroundHttpResponseImpl.h b/Engine/Source/Runtime/Online/BackgroundHTTP/Public/BackgroundHttpResponseImpl.h
--- a/Engine/Source/Runtime/Online/BackgroundHTTP/Public/BackgroundHttpResponseImpl.h
+++ b/Engine/Source/Runtime/Online/BackgroundHTTP/Public/BackgroundHttpResponseImpl.h
@@ -14,11 +14,12 @@ class FBackgroundHttpResponseImpl
 {
 public:
 	FBackgroundHttpResponseImpl();
-	virtual ~FBackgroundHttpResponseImpl() = default;
+	FBackgroundHttpResponseImpl(FString InTempContentFilePath, int32 InResponseCode);
+	virtual ~FBackgroundHttpResponseImpl() override = default;
 
 	//IHttpBackgroundResponse
-	virtual int32 GetResponseCode() const;
-	virtual const FString& GetTempContentFilePath() const;
+	virtual int32 GetResponseCode() const override;
+	virtual const FString& GetTempContentFilePath() const override;
 
 protected:	
 	FString TempContentFilePath;
