Fix for ForegroundServiceDidNotStopInTimeException
- doWork as blocking potentially forever; added a watchdog (5.5 hours)
- similar watchdog on DownloadManager worker since it reschedules itself indefinitely
- bForceStopped volatile since set in another thread
- reduced the timeout for OKHttp so it doesn't spend so long before retry

#android

#Synced-CL -1

diff --git a/Engine/Plugins/Online/Android/AndroidFetchBackgroundDownload/Source/AndroidFetchBackgroundDownload/Java/download/DownloadManager.java b/Engine/Plugins/Online/Android/AndroidFetchBackgroundDownload/Source/AndroidFetchBackgroundDownload/Java/download/DownloadManager.java
--- a/Engine/Plugins/Online/Android/AndroidFetchBackgroundDownload/Source/AndroidFetchBackgroundDownload/Java/download/DownloadManager.java
+++ b/Engine/Plugins/Online/Android/AndroidFetchBackgroundDownload/Source/AndroidFetchBackgroundDownload/Java/download/DownloadManager.java
@@ -671,6 +671,9 @@ public class DownloadManager
 	{
 		CallbackToFutureAdapter.Completer<ListenableWorker.Result> Completer;
 
+		long StartTimeMs;
+		static final long MaxRuntimeMs = TimeUnit.HOURS.toMillis(5) + TimeUnit.MINUTES.toMillis(30); // 5.5 hrs, before Android's 6hr foreground service limit
+
 		DownloadNotificationDescription NotificationDescription;
 		
 		PendingIntent PendingCancelIntent;
@@ -944,6 +947,8 @@ public class DownloadManager
 		@Override
 		public @NonNull ListenableFuture<Result> startWork()
 		{
+			StartTimeMs = System.currentTimeMillis();
+
 			ListenableFuture<ListenableWorker.Result> future = CallbackToFutureAdapter.getFuture(this);
 
 			final Context context = getApplicationContext();
@@ -1017,6 +1022,13 @@ public class DownloadManager
 		{
 			if (Task != null && Task.Worker == this)
 			{
+				if (System.currentTimeMillis() - StartTimeMs > MaxRuntimeMs)
+				{
+					Log.error("Worker exceeded maximum runtime of 5.5 hours. Forcing retry to avoid foreground service timeout.");
+					Task.AllComplete();
+					return;
+				}
+
 				@SuppressLint("InlinedApi")
 				ListenableFuture<Void> future = setForegroundAsync(new ForegroundInfo(NotificationDescription.NotificationID, CreateNotification(getApplicationContext()), android.content.pm.ServiceInfo.FOREGROUND_SERVICE_TYPE_DATA_SYNC));
 				future.addListener(() ->
diff --git a/Engine/Plugins/Online/Android/AndroidFetchBackgroundDownload/Source/AndroidFetchBackgroundDownload/Java/download/UEDownloadWorker.java b/Engine/Plugins/Online/Android/AndroidFetchBackgroundDownload/Source/AndroidFetchBackgroundDownload/Java/download/UEDownloadWorker.java
--- a/Engine/Plugins/Online/Android/AndroidFetchBackgroundDownload/Source/AndroidFetchBackgroundDownload/Java/download/UEDownloadWorker.java
+++ b/Engine/Plugins/Online/Android/AndroidFetchBackgroundDownload/Source/AndroidFetchBackgroundDownload/Java/download/UEDownloadWorker.java
@@ -224,10 +224,19 @@ public class UEDownloadWorker extends UEWorker implements DownloadProgressListen
 
 		//Enter actual loop until work is finished
 		Log.verbose("Entering OnWorkerStart Loop");
-		try 
+		final long startTimeMs = System.currentTimeMillis();
+		final long maxRuntimeMs = TimeUnit.HOURS.toMillis(5) + TimeUnit.MINUTES.toMillis(30); // 5.5 hrs, before Android's 6hr foreground service limit
+		try
 		{
-			while (bReceivedResult == false)
+			while (bReceivedResult == false && !isStopped())
 			{
+				if (System.currentTimeMillis() - startTimeMs > maxRuntimeMs)
+				{
+					Log.error("OnWorkerStart exceeded maximum runtime of 5.5 hours. Forcing retry to avoid foreground service timeout.");
+					SetWorkResult_Retry();
+					break;
+				}
+
 				// check if taskfile has been changed
 				String RequestedTaskFilename = preferencesBackground.getString("TaskFilename", CurrentTaskFilename);
 				if (RequestedTaskFilename.equals(CurrentTaskFilename))
@@ -996,7 +1005,7 @@ public class UEDownloadWorker extends UEWorker implements DownloadProgressListen
 	private boolean bLostNetwork = false;
 	private boolean bAirplaneMode = false;
 	private MeteredType meteredType = MeteredType.UNMETERED;
-	private boolean bForceStopped = false;
+	private volatile boolean bForceStopped = false;
 	private DownloadQueueDescription QueueDescription = null;
 	private volatile boolean bHasEnqueueHappened = false;
 	static volatile FetchManagerInterface mFetchManager;
diff --git a/Engine/Plugins/Online/Android/AndroidFetchBackgroundDownload/Source/AndroidFetchBackgroundDownload/Java/download/asyncdownloader/AsyncDownloader.java b/Engine/Plugins/Online/Android/AndroidFetchBackgroundDownload/Source/AndroidFetchBackgroundDownload/Java/download/asyncdownloader/AsyncDownloader.java
--- a/Engine/Plugins/Online/Android/AndroidFetchBackgroundDownload/Source/AndroidFetchBackgroundDownload/Java/download/asyncdownloader/AsyncDownloader.java
+++ b/Engine/Plugins/Online/Android/AndroidFetchBackgroundDownload/Source/AndroidFetchBackgroundDownload/Java/download/asyncdownloader/AsyncDownloader.java
@@ -37,7 +37,7 @@ public class AsyncDownloader {
 				.dispatcher(Dispatcher)
 				.callTimeout(Duration.ofMinutes(30))
 				.connectTimeout(Duration.ofSeconds(20))
-				.readTimeout(Duration.ofMinutes(30))
+				.readTimeout(Duration.ofMinutes(5))
 				.protocols(AllowHTTP2 ? List.of(Protocol.HTTP_2, Protocol.HTTP_1_1) : List.of(Protocol.HTTP_1_1))
 				.build();
 	}
