[Mobile] Fix previous "re-queue" causing consistent initial download stalls.

Added API entrypoints to query the active downloads from Java (through JNI) and FBackgroundURLSessionHandler.
From that list of internal request IDs, the throughput watchdog can correctly target the slow ones.
Not mistaking queued requests for those currently downloading at a very slow throughput.

The whole watchdog addition is still cvar-gated.

The throughput threshold is configurable through the cvar, so we can adjust it if needed.

#rnx
#Synced-CL 52374856

#Synced-CL 52423882

diff --git a/Engine/Plugins/Online/Android/AndroidFetchBackgroundDownload/Source/AndroidFetchBackgroundDownload/Java/download/DownloadManager.java b/Engine/Plugins/Online/Android/AndroidFetchBackgroundDownload/Source/AndroidFetchBackgroundDownload/Java/download/DownloadManager.java
--- a/Engine/Plugins/Online/Android/AndroidFetchBackgroundDownload/Source/AndroidFetchBackgroundDownload/Java/download/DownloadManager.java
+++ b/Engine/Plugins/Online/Android/AndroidFetchBackgroundDownload/Source/AndroidFetchBackgroundDownload/Java/download/DownloadManager.java
@@ -1239,4 +1239,15 @@ public class DownloadManager
 			}
 		});
 	}
+
+	@Keep
+	public static String GetActiveDestinations()
+	{
+		if (Task != null && Task.FetchManager != null)
+		{
+			String[] Destinations = Task.FetchManager.GetRunningDestinationLocations();
+			return String.join("\n", Destinations);
+		}
+		return "";
+	}
 }
\ No newline at end of file
diff --git a/Engine/Plugins/Online/Android/AndroidFetchBackgroundDownload/Source/AndroidFetchBackgroundDownload/Java/download/asyncdownloader/AsyncDownloader.java b/Engine/Plugins/Online/Android/AndroidFetchBackgroundDownload/Source/AndroidFetchBackgroundDownload/Java/download/asyncdownloader/AsyncDownloader.java
--- a/Engine/Plugins/Online/Android/AndroidFetchBackgroundDownload/Source/AndroidFetchBackgroundDownload/Java/download/asyncdownloader/AsyncDownloader.java
+++ b/Engine/Plugins/Online/Android/AndroidFetchBackgroundDownload/Source/AndroidFetchBackgroundDownload/Java/download/asyncdownloader/AsyncDownloader.java
@@ -208,6 +208,22 @@ public class AsyncDownloader {
 		return Math.min(1f, (float)DownloadedBytes / (float)TotalBytes);
 	}
 
+	public synchronized String[] GetRunningDestinationLocations() {
+		ArrayList<String> result = new ArrayList<>();
+		for (Map.Entry<String, DownloadEntry> entry : Entries.entrySet()) {
+			DownloadEntry e = entry.getValue();
+			// EState.RUNNING is set at enqueue time, before OkHttp's Dispatcher assigns a connection.
+			// With MaxActiveDownloads, entries beyond the limit are queued internally by OkHttp but still
+			// marked RUNNING. Use two signals to confirm a real connection:
+			//   - bResponseReceived: HTTP response headers received (connection established)
+			//   - DownloadedSoFar > 0: body data is flowing
+			if (e.State == EState.RUNNING && (e.bResponseReceived || e.DownloadedSoFar > 0)) {
+				result.add(entry.getKey());
+			}
+		}
+		return result.toArray(new String[0]);
+	}
+
 	// ---------------------------------------------------------------------------------------------------------------------
 	// private data + methods
 	// ---------------------------------------------------------------------------------------------------------------------
@@ -280,6 +296,7 @@ public class AsyncDownloader {
 
 			Request Request = rb.build();
 			Entry.State = EState.RUNNING;
+			Entry.bResponseReceived = false;
 			Entry.DownloadedSoFar = CurrentLen;
 			Entry.Listener.OnStart(Entry.Request.Id);
 
@@ -294,6 +311,7 @@ public class AsyncDownloader {
 
 				@Override public void onResponse(@NonNull Call Call, @NonNull Response Response) {
 					//Log.debug("onResponse: begin downloading | Id=" + Entry.Request.Id);
+					Entry.bResponseReceived = true;
 					try (Response response = Response) {
 						int Code = response.code();
 						if (!response.isSuccessful()) {
@@ -460,6 +478,7 @@ public class AsyncDownloader {
 		volatile int Attempts = 0;				// total attempts across all URLs
 		volatile int UrlIndex = 0;				// which URL we’ll try next
 		volatile long DownloadedSoFar = 0L;		// tracked for progress (file length + session)
+		volatile boolean bResponseReceived = false; // true once HTTP response headers are received (real connection)
 		volatile boolean UserPaused = false;	// true if pause(Id) was requested
 		volatile boolean Cancelled = false;
 
diff --git a/Engine/Plugins/Online/Android/AndroidFetchBackgroundDownload/Source/AndroidFetchBackgroundDownload/Java/download/fetch/FetchManagerNew.java b/Engine/Plugins/Online/Android/AndroidFetchBackgroundDownload/Source/AndroidFetchBackgroundDownload/Java/download/fetch/FetchManagerNew.java
--- a/Engine/Plugins/Online/Android/AndroidFetchBackgroundDownload/Source/AndroidFetchBackgroundDownload/Java/download/fetch/FetchManagerNew.java
+++ b/Engine/Plugins/Online/Android/AndroidFetchBackgroundDownload/Source/AndroidFetchBackgroundDownload/Java/download/fetch/FetchManagerNew.java
@@ -224,6 +224,14 @@ public class FetchManagerNew implements FetchManagerInterface
 		AsyncDownloader.Retry(DestinationLocation);
 	}
 
+	public synchronized String[] GetRunningDestinationLocations()
+	{
+		if (AsyncDownloader == null) {
+			return new String[0];
+		}
+		return AsyncDownloader.GetRunningDestinationLocations();
+	}
+
 	@Override
 	public synchronized void RequestGroupProgressUpdate(int GroupID, DownloadProgressListener ListenerToUpdate)
 	{
diff --git a/Engine/Plugins/Online/Android/AndroidFetchBackgroundDownload/Source/AndroidFetchBackgroundDownload/Private/AndroidDownloadManager.cpp b/Engine/Plugins/Online/Android/AndroidFetchBackgroundDownload/Source/AndroidFetchBackgroundDownload/Private/AndroidDownloadManager.cpp
--- a/Engine/Plugins/Online/Android/AndroidFetchBackgroundDownload/Source/AndroidFetchBackgroundDownload/Private/AndroidDownloadManager.cpp
+++ b/Engine/Plugins/Online/Android/AndroidFetchBackgroundDownload/Source/AndroidFetchBackgroundDownload/Private/AndroidDownloadManager.cpp
@@ -47,6 +47,7 @@ namespace UE::Jni::Download
 		inline static TMember<FDownloadManager, void (Java::Lang::FString* DestinationLocation)> Pause;
 		inline static TMember<FDownloadManager, void (Java::Lang::FString* DestinationLocation)> Resume;
 		inline static TMember<FDownloadManager, void (Java::Lang::FString* DestinationLocation)> Cancel;
+		inline static TMember<FDownloadManager, Java::Lang::FString* ()> GetActiveDestinations;
 		
 		static void JNICALL OnProgress(JNIEnv* env, Java::Lang::TClass<FDownloadManager>* clazz, Java::Lang::FString* DestinationLocation, jlong TotalBytesWritten, jlong BytesWrittenSinceLastCall)
 		{
@@ -77,7 +78,8 @@ namespace UE::Jni::Download
 			UE_JNI_MEMBER(Enqueue),
 			UE_JNI_MEMBER(Pause),
 			UE_JNI_MEMBER(Resume),
-			UE_JNI_MEMBER(Cancel)
+			UE_JNI_MEMBER(Cancel),
+			UE_JNI_MEMBER(GetActiveDestinations)
 		};
 		
 		static constexpr FNativeMethod NativeMethods[]
@@ -321,6 +323,16 @@ namespace UE::Online::Download::Android
 		FileHashHelper->SaveData();
 	}
 	
+	void FDownloadManager::FRequest::RotateURLList()
+	{
+		if (URLList.Num() > 1)
+		{
+			FString First = MoveTemp(URLList[0]);
+			URLList.RemoveAt(0, EAllowShrinking::No);
+			URLList.Add(MoveTemp(First));
+		}
+	}
+
 	void FDownloadManager::EnqueueRequest(FRequest& _Request)
 	{
 		TSharedPtr<FRequest> Request = SharedThis(&_Request);
@@ -391,6 +403,51 @@ namespace UE::Online::Download::Android
 		}
 	}
 	
+	void FDownloadManager::RequeueRequest(const FBackgroundHttpRequestPtr GenericRequest)
+	{
+		ensureAlwaysMsgf(IsInGameThread(), TEXT("RequeueRequest called from unexpected thread!"));
+
+		FRequest& Request = static_cast<FRequest&>(*GenericRequest);
+
+		UE_LOG(LogBackgroundHttpManager, Display, TEXT("RequeueRequest - RequestID:%s"), *Request.RequestID);
+
+		// Cancel existing Java-side download
+		CancelRequest(Request);
+
+		// Rotate URL list so the retry hits a different CDN
+		Request.RotateURLList();
+
+		// Re-enqueue with fresh Java-side download
+		EnqueueRequest(Request);
+	}
+
+	void FDownloadManager::GetActiveDownloadRequestIDs(TSet<FString>& OutActiveIDs) const
+	{
+		OutActiveIDs.Reset();
+
+#if USE_ANDROID_JNI
+		using namespace UE;
+		FScopedJavaObject<Jni::Java::Lang::FString*> ActiveDestinations = Jni::Download::FDownloadManager::GetActiveDestinations(Jni::Class<Jni::Download::FDownloadManager>);
+		if (ActiveDestinations)
+		{
+			FString DestinationList = Jni::Env.Marshal(*ActiveDestinations);
+			if (!DestinationList.IsEmpty())
+			{
+				TArray<FString> Destinations;
+				DestinationList.ParseIntoArray(Destinations, TEXT("\n"));
+
+				for (const FString& Destination : Destinations)
+				{
+					if (TSharedPtr<FRequest> Request = Requests.FindRef(Destination))
+					{
+						OutActiveIDs.Add(Request->GetRequestID());
+					}
+				}
+			}
+		}
+#endif
+	}
+
 	void FDownloadManager::OnProgress(const FString& DestinationLocation, int64 TotalBytesWritten, int64 BytesWrittenSinceLastCall) const
 	{
 		UE_LOG(LogBackgroundHttpManager, VeryVerbose, TEXT("Download Progress... DestinationLocation:%s BytesWrittenSinceLastCall:%lld TotalBytesWritten:%lld"), *DestinationLocation, BytesWrittenSinceLastCall, TotalBytesWritten);
diff --git a/Engine/Plugins/Online/Android/AndroidFetchBackgroundDownload/Source/AndroidFetchBackgroundDownload/Public/AndroidDownloadManager.h b/Engine/Plugins/Online/Android/AndroidFetchBackgroundDownload/Source/AndroidFetchBackgroundDownload/Public/AndroidDownloadManager.h
--- a/Engine/Plugins/Online/Android/AndroidFetchBackgroundDownload/Source/AndroidFetchBackgroundDownload/Public/AndroidDownloadManager.h
+++ b/Engine/Plugins/Online/Android/AndroidFetchBackgroundDownload/Source/AndroidFetchBackgroundDownload/Public/AndroidDownloadManager.h
@@ -22,6 +22,8 @@ namespace UE::Online::Download::Android
 			virtual void PauseRequest() override;
 			virtual void ResumeRequest() override;
 
+			void RotateURLList();
+
 		private:
 			int GetPriorityAsAndroidPriority() const;
 			
@@ -41,7 +43,9 @@ namespace UE::Online::Download::Android
 		
 		virtual void AddRequest(const FBackgroundHttpRequestPtr Request) override;
 		virtual void RemoveRequest(const FBackgroundHttpRequestPtr Request) override;
-		
+		virtual void RequeueRequest(const FBackgroundHttpRequestPtr Request) override;
+		virtual void GetActiveDownloadRequestIDs(TSet<FString>& OutActiveIDs) const override;
+
 		virtual void DeleteAllTemporaryFiles() override;
 		virtual bool IsGenericImplementation() const override { return false; }
 		
diff --git a/Engine/Source/Runtime/Apple/BackgroundHttpIOS/Private/BackgroundURLSessionHandler.cpp b/Engine/Source/Runtime/Apple/BackgroundHttpIOS/Private/BackgroundURLSessionHandler.cpp
--- a/Engine/Source/Runtime/Apple/BackgroundHttpIOS/Private/BackgroundURLSessionHandler.cpp
+++ b/Engine/Source/Runtime/Apple/BackgroundHttpIOS/Private/BackgroundURLSessionHandler.cpp
@@ -378,6 +378,7 @@ enum class EBackgroundNSURLCDNInfoResponse : uint32 // Beware these constants de
 - (void)SetPriority:(float)Priority ForDownload:(NSUInteger)DownloadId;
 - (void)SetCurrentDownloadedBytes:(uint64)DownloadedBytes ForTask:(NSURLSessionDownloadTask*)Task;
 - (uint64)GetCurrentDownloadedBytes:(NSUInteger)DownloadId;
+- (void)GetActiveDownloadIDs:(TSet<uint64>&)OutActiveIDs;
 - (void)RecreateDownload:(NSUInteger)DownloadId ShouldResetRetryCount:(bool)ResetRetryCount;
 - (void)RecreateDownloads;
 #if !UE_BUILD_SHIPPING
@@ -1230,6 +1231,29 @@ static constexpr NSInteger TypicalSizeForDownloadResponseHeader = 400; // bytes
 	return 0;
 }
 
+- (void)GetActiveDownloadIDs:(TSet<uint64>&)OutActiveIDs
+{
+	@synchronized(_AllDownloads)
+	{
+		[_AllDownloads enumerateKeysAndObjectsUsingBlock:^(NSNumber* Key, NSURLSessionDownloadTask* Task, BOOL* IterStop)
+		{
+			// NSURLSessionTaskStateRunning means "started, not suspended/cancelled" — it does NOT mean
+			// the task has an active TCP connection. With MaximumConnectionsPerHost, NSURLSession internally
+			// queues tasks beyond the connection limit and still reports them as Running.
+			// Use two signals to determine if a task has a real connection:
+			//   - countOfBytesReceived > 0: body data is flowing
+			//   - countOfBytesExpectedToReceive > 0 and != unknown: response headers received (connection
+			//     established, server responded) even if no body data yet
+			const bool bHasReceivedData = Task.countOfBytesReceived > 0;
+			const bool bHasReceivedHeaders = Task.countOfBytesExpectedToReceive > 0 && Task.countOfBytesExpectedToReceive != NSURLSessionTransferSizeUnknown;
+			if (Task.state == NSURLSessionTaskStateRunning && (bHasReceivedData || bHasReceivedHeaders))
+			{
+				OutActiveIDs.Add(Key.unsignedLongLongValue);
+			}
+		}];
+	}
+}
+
 - (void)RecreateDownload:(NSUInteger)DownloadId ShouldResetRetryCount:(bool)ResetRetryCount
 {
 	UE_DNLD_LOG(@"RecreateDownload for DownloadId %lu", DownloadId);
@@ -1947,6 +1971,15 @@ uint64 FBackgroundURLSessionHandler::GetCurrentDownloadedBytes(const uint64 Down
 	}
 }
 
+void FBackgroundURLSessionHandler::GetActiveDownloadIDs(TSet<uint64>& OutActiveIDs)
+{
+	@autoreleasepool
+	{
+		OutActiveIDs.Reset();
+		[[FBackgroundNSURLSession Shared] GetActiveDownloadIDs:OutActiveIDs];
+	}
+}
+
 bool FBackgroundURLSessionHandler::IsDownloadFinished(const uint64 DownloadId, int32& OutResultHTTPCode, FString& OutTemporaryFilePath)
 {
 	@autoreleasepool
diff --git a/Engine/Source/Runtime/Apple/BackgroundHttpIOS/Public/BackgroundURLSessionHandler.h b/Engine/Source/Runtime/Apple/BackgroundHttpIOS/Public/BackgroundURLSessionHandler.h
--- a/Engine/Source/Runtime/Apple/BackgroundHttpIOS/Public/BackgroundURLSessionHandler.h
+++ b/Engine/Source/Runtime/Apple/BackgroundHttpIOS/Public/BackgroundURLSessionHandler.h
@@ -77,6 +77,9 @@ public:
 
 	static uint64 GetCurrentDownloadedBytes(const uint64 DownloadId);
 
+	// Returns download IDs of tasks currently in NSURLSessionTaskStateRunning state
+	static void GetActiveDownloadIDs(TSet<uint64>& OutActiveIDs);
+
 	static bool IsDownloadFinished(const uint64 DownloadId, int32& OutResultHTTPCode, FString& OutTemporaryFilePath);
 
 	// To be used by app delegate, call it from handleEventsForBackgroundURLSession.
diff --git a/Engine/Source/Runtime/Online/BackgroundHTTP/Private/BackgroundHttpManagerImpl.cpp b/Engine/Source/Runtime/Online/BackgroundHTTP/Private/BackgroundHttpManagerImpl.cpp
--- a/Engine/Source/Runtime/Online/BackgroundHTTP/Private/BackgroundHttpManagerImpl.cpp
+++ b/Engine/Source/Runtime/Online/BackgroundHTTP/Private/BackgroundHttpManagerImpl.cpp
@@ -356,6 +356,16 @@ void FBackgroundHttpManagerImpl::CleanUpDataAfterCompletingRequest(const FBackgr
 	}
 }
 
+void FBackgroundHttpManagerImpl::SetCellularPreference(int32 Value)
+{
+
+}
+
+void FBackgroundHttpManagerImpl::SetLimitedDataPreference(int32 Value)
+{
+	
+}
+
 bool FBackgroundHttpManagerImpl::AssociateWithAnyExistingRequest(const FBackgroundHttpRequestPtr Request)
 {
 	bool bDidAssociateWithExistingRequest = false;
diff --git a/Engine/Source/Runtime/Online/BackgroundHTTP/Private/BackgroundHttpRequestImpl.cpp b/Engine/Source/Runtime/Online/BackgroundHTTP/Private/BackgroundHttpRequestImpl.cpp
--- a/Engine/Source/Runtime/Online/BackgroundHTTP/Private/BackgroundHttpRequestImpl.cpp
+++ b/Engine/Source/Runtime/Online/BackgroundHTTP/Private/BackgroundHttpRequestImpl.cpp
@@ -30,6 +30,12 @@ void FBackgroundHttpRequestImpl::CancelRequest()
 	FBackgroundHttpModule::Get().GetBackgroundHttpManager()->RemoveRequest(SharedThis(this));
 }
 
+void FBackgroundHttpRequestImpl::RequeueRequest()
+{
+	UE_LOGF(LogBackgroundHttpRequest, Display, "RequeueRequest - RequestID:%ls", *GetRequestID());
+	FBackgroundHttpModule::Get().GetBackgroundHttpManager()->RequeueRequest(SharedThis(this));
+}
+
 void FBackgroundHttpRequestImpl::PauseRequest()
 {
     //for now a pause is just wrapping a cancel in the general case
@@ -146,19 +152,6 @@ const TArray<FString>& FBackgroundHttpRequestImpl::GetURLList() const
 	return URLList;
 }
 
-void FBackgroundHttpRequestImpl::RotateURLList()
-{
-	// Only safe to call when the request is not being processed (i.e., after CancelRequest and before ProcessRequest).
-	// Callers must ensure this precondition.
-	if (URLList.Num() > 1)
-	{
-		UE_LOG(LogBackgroundHttpRequest, Verbose, TEXT("Rotating URL list for request %s (moving %s to end)"), *GetRequestID(), *URLList[0]);
-		FString First = MoveTemp(URLList[0]);
-		URLList.RemoveAt(0, EAllowShrinking::No);
-		URLList.Add(MoveTemp(First));
-	}
-}
-
 void FBackgroundHttpRequestImpl::SetExpectedResultSize(const uint64 ExpectedSize)
 {
 	ExpectedResultSize = ExpectedSize;
diff --git a/Engine/Source/Runtime/Online/BackgroundHTTP/Private/IOS/ApplePlatformBackgroundHttpManager.cpp b/Engine/Source/Runtime/Online/BackgroundHTTP/Private/IOS/ApplePlatformBackgroundHttpManager.cpp
--- a/Engine/Source/Runtime/Online/BackgroundHTTP/Private/IOS/ApplePlatformBackgroundHttpManager.cpp
+++ b/Engine/Source/Runtime/Online/BackgroundHTTP/Private/IOS/ApplePlatformBackgroundHttpManager.cpp
@@ -80,6 +80,66 @@ void FApplePlatformBackgroundHttpManager::RemoveRequest(const FBackgroundHttpReq
 	PendingRemoveRequests.Add(GenericRequest);
 }
 
+void FApplePlatformBackgroundHttpManager::RequeueRequest(const FBackgroundHttpRequestPtr GenericRequest)
+{
+	ensureAlwaysMsgf(IsInGameThread(), TEXT("RequeueRequest called from unexpected thread!"));
+
+	FAppleBackgroundHttpRequestPtr Request = StaticCastSharedPtr<FApplePlatformBackgroundHttpRequest>(GenericRequest);
+	if (!ensureAlwaysMsgf(Request.IsValid(), TEXT("Re-queuing a non-Apple background request! This is not supported.")))
+	{
+		return;
+	}
+
+	// Cancel the old NSURLSession download immediately (NOT deferred).
+	// This avoids the race where deferred removal would cancel the new download instead.
+	const uint64 OldDownloadId = Request->GetInternalDownloadId();
+	if (OldDownloadId != FBackgroundURLSessionHandler::InvalidDownloadId)
+	{
+		FBackgroundURLSessionHandler::CancelDownload(OldDownloadId);
+	}
+
+	// Rotate URL list so the retry hits a different CDN
+	Request->RotateURLList();
+
+	// Create a fresh download with rotated URLs — new DownloadId, fresh TaskData, fresh retry counts
+	const uint64 NewDownloadId = FBackgroundURLSessionHandler::CreateOrFindDownload(
+		Request->GetURLList(),
+		Request->GetNSURLSessionPriority(),
+		GetFileHashHelper(),
+		Request->GetExpectedResultSize(),
+		Request->GetRequestID());
+	Request->SetInternalDownloadId(NewDownloadId);
+
+	if (NewDownloadId == FBackgroundURLSessionHandler::InvalidDownloadId)
+	{
+		UE_LOG(LogBackgroundHttpManager, Error, TEXT("RequeueRequest failed to create new download for RequestID:%s"), *Request->GetRequestID());
+		return;
+	}
+
+	// Reset progress tracking for the new download
+	Request->ResetProgressTracking();
+
+	UE_LOG(LogBackgroundHttpManager, Display, TEXT("RequeueRequest succeeded - RequestID:%s OldDownloadId:%llu NewDownloadId:%llu"), *Request->GetRequestID(), OldDownloadId, NewDownloadId);
+}
+
+void FApplePlatformBackgroundHttpManager::GetActiveDownloadRequestIDs(TSet<FString>& OutActiveIDs) const
+{
+	OutActiveIDs.Reset();
+
+	TSet<uint64> ActiveDownloadIDs;
+	FBackgroundURLSessionHandler::GetActiveDownloadIDs(ActiveDownloadIDs);
+
+	FRWScopeLock ScopeLock(ActiveRequestLock, SLT_ReadOnly);
+	for (const FBackgroundHttpRequestPtr& GenericRequest : ActiveRequests)
+	{
+		FAppleBackgroundHttpRequestPtr Request = StaticCastSharedPtr<FApplePlatformBackgroundHttpRequest>(GenericRequest);
+		if (Request.IsValid() && ActiveDownloadIDs.Contains(Request->GetInternalDownloadId()))
+		{
+			OutActiveIDs.Add(Request->GetRequestID());
+		}
+	}
+}
+
 void FApplePlatformBackgroundHttpManager::SetCellularPreference(int32 Value)
 {
 	FBackgroundURLSessionHandler::AllowCellular(Value > 0);
diff --git a/Engine/Source/Runtime/Online/BackgroundHTTP/Private/IOS/ApplePlatformBackgroundHttpRequest.cpp b/Engine/Source/Runtime/Online/BackgroundHTTP/Private/IOS/ApplePlatformBackgroundHttpRequest.cpp
--- a/Engine/Source/Runtime/Online/BackgroundHTTP/Private/IOS/ApplePlatformBackgroundHttpRequest.cpp
+++ b/Engine/Source/Runtime/Online/BackgroundHTTP/Private/IOS/ApplePlatformBackgroundHttpRequest.cpp
@@ -97,6 +97,22 @@ void FApplePlatformBackgroundHttpRequest::UpdateProgress()
 	}
 }
 
+void FApplePlatformBackgroundHttpRequest::ResetProgressTracking()
+{
+	LastReportedDownloadedBytes = 0;
+	MetricsWaitCounter = -1;
+}
+
+void FApplePlatformBackgroundHttpRequest::RotateURLList()
+{
+	if (URLList.Num() > 1)
+	{
+		FString First = MoveTemp(URLList[0]);
+		URLList.RemoveAt(0, EAllowShrinking::No);
+		URLList.Add(MoveTemp(First));
+	}
+}
+
 float FApplePlatformBackgroundHttpRequest::BackgroundRequestPriorityToNSURLSessionPriority(const EBackgroundHTTPPriority Priority)
 {
 	switch (Priority)
diff --git a/Engine/Source/Runtime/Online/BackgroundHTTP/Public/BackgroundHttpManagerImpl.h b/Engine/Source/Runtime/Online/BackgroundHTTP/Public/BackgroundHttpManagerImpl.h
--- a/Engine/Source/Runtime/Online/BackgroundHTTP/Public/BackgroundHttpManagerImpl.h
+++ b/Engine/Source/Runtime/Online/BackgroundHTTP/Public/BackgroundHttpManagerImpl.h
@@ -44,8 +44,9 @@ public:
 	BACKGROUNDHTTP_API virtual FString GetTempFileLocationForURL(const FString& URL, const FString& RequestID) override;
 	
 	BACKGROUNDHTTP_API virtual void CleanUpDataAfterCompletingRequest(const FBackgroundHttpRequestPtr Request) override;
-	virtual void SetCellularPreference(int32 Value) override {}
-	virtual void SetLimitedDataPreference(int32 Value) override {}
+
+	BACKGROUNDHTTP_API virtual void SetCellularPreference(int32 Value) override;
+	BACKGROUNDHTTP_API virtual void SetLimitedDataPreference(int32 Value) override;
 	
 	//FTSTickerObjectBase implementation
 	BACKGROUNDHTTP_API virtual bool Tick(float DeltaTime) override;
@@ -88,7 +89,7 @@ protected:
 
 	/** List of Background Http requests that are actively being processed **/
 	TArray<FBackgroundHttpRequestPtr> ActiveRequests;
-	FRWLock ActiveRequestLock;
+	mutable FRWLock ActiveRequestLock;
 
 	/** Count of how many requests we have active **/
 	volatile int NumCurrentlyActiveRequests;
diff --git a/Engine/Source/Runtime/Online/BackgroundHTTP/Public/BackgroundHttpRequestImpl.h b/Engine/Source/Runtime/Online/BackgroundHTTP/Public/BackgroundHttpRequestImpl.h
--- a/Engine/Source/Runtime/Online/BackgroundHTTP/Public/BackgroundHttpRequestImpl.h
+++ b/Engine/Source/Runtime/Online/BackgroundHTTP/Public/BackgroundHttpRequestImpl.h
@@ -28,11 +28,11 @@ public:
 	// IBackgroundHttpRequest
 	BACKGROUNDHTTP_API virtual bool ProcessRequest() override;
 	BACKGROUNDHTTP_API virtual void CancelRequest() override;
+	BACKGROUNDHTTP_API virtual void RequeueRequest() override;
     BACKGROUNDHTTP_API virtual void PauseRequest() override;
     BACKGROUNDHTTP_API virtual void ResumeRequest() override;
     BACKGROUNDHTTP_API virtual void SetURLAsList(const TArray<FString>& URLs, int NumRetriesToAttempt) override;
 	BACKGROUNDHTTP_API virtual const TArray<FString>& GetURLList() const override;
-	BACKGROUNDHTTP_API virtual void RotateURLList() override;
 	BACKGROUNDHTTP_API virtual void SetExpectedResultSize(const uint64 ExpectedSize) override;
 	BACKGROUNDHTTP_API virtual uint64 GetExpectedResultSize() const override;
 	BACKGROUNDHTTP_API virtual void SetCompleteNotification(FBackgroundHttpNotificationObjectPtr DownloadCompleteNotificationObjectIn) override;
diff --git a/Engine/Source/Runtime/Online/BackgroundHTTP/Public/IOS/ApplePlatformBackgroundHttpManager.h b/Engine/Source/Runtime/Online/BackgroundHTTP/Public/IOS/ApplePlatformBackgroundHttpManager.h
--- a/Engine/Source/Runtime/Online/BackgroundHTTP/Public/IOS/ApplePlatformBackgroundHttpManager.h
+++ b/Engine/Source/Runtime/Online/BackgroundHTTP/Public/IOS/ApplePlatformBackgroundHttpManager.h
@@ -25,6 +25,8 @@ public:
 	virtual void Initialize(bool bClearPreExistingRequestsAtStartup) override;
 	virtual void AddRequest(const FBackgroundHttpRequestPtr Request) override;
 	virtual void RemoveRequest(const FBackgroundHttpRequestPtr Request) override;
+	virtual void RequeueRequest(const FBackgroundHttpRequestPtr Request) override;
+	virtual void GetActiveDownloadRequestIDs(TSet<FString>& OutActiveIDs) const override;
 	virtual void SetCellularPreference(int32 Value) override;
 	virtual void SetLimitedDataPreference(int32 Value) override;
 	virtual bool IsGenericImplementation() const override { return false; }
diff --git a/Engine/Source/Runtime/Online/BackgroundHTTP/Public/IOS/ApplePlatformBackgroundHttpRequest.h b/Engine/Source/Runtime/Online/BackgroundHTTP/Public/IOS/ApplePlatformBackgroundHttpRequest.h
--- a/Engine/Source/Runtime/Online/BackgroundHTTP/Public/IOS/ApplePlatformBackgroundHttpRequest.h
+++ b/Engine/Source/Runtime/Online/BackgroundHTTP/Public/IOS/ApplePlatformBackgroundHttpRequest.h
@@ -27,6 +27,8 @@ protected:
 	uint64 GetInternalDownloadId() const {return DownloadId;}
 
 	void UpdateProgress();
+	void ResetProgressTracking();
+	void RotateURLList();
 
 	static float BackgroundRequestPriorityToNSURLSessionPriority(const EBackgroundHTTPPriority Priority);
 	float GetNSURLSessionPriority() const
diff --git a/Engine/Source/Runtime/Online/BackgroundHTTP/Public/Interfaces/IBackgroundHttpManager.h b/Engine/Source/Runtime/Online/BackgroundHTTP/Public/Interfaces/IBackgroundHttpManager.h
--- a/Engine/Source/Runtime/Online/BackgroundHTTP/Public/Interfaces/IBackgroundHttpManager.h
+++ b/Engine/Source/Runtime/Online/BackgroundHTTP/Public/Interfaces/IBackgroundHttpManager.h
@@ -38,6 +38,15 @@ public:
 	 */
 	virtual void RemoveRequest(const FBackgroundHttpRequestPtr Request) = 0;
 
+	/**
+	 * Re-queues an active request by canceling the current platform download and creating a fresh one.
+	 * Unlike CancelRequest+ProcessRequest, this is safe on all platforms as it avoids deferred removal races.
+	 * The URL list is rotated so the retry hits a different CDN.
+	 *
+	 * @param Request - the request object to re-queue
+	 */
+	virtual void RequeueRequest(const FBackgroundHttpRequestPtr Request) {}
+
 	/**
 	* Function to remove all temporary files used by the system to store completed downloads.
 	* This should be called when there are no active background downloads that we care
@@ -87,7 +96,14 @@ public:
 	 * Setting limited data network preference
 	 */
 	virtual void SetLimitedDataPreference(int32 Value) = 0;
-	
+
+	/**
+	 * Returns the set of RequestIDs that are actively being downloaded by the platform.
+	 * Requests queued by the platform (waiting for a concurrency slot) are excluded.
+	 * Default implementation returns empty set (treats no requests as confirmed active).
+	 */
+	virtual void GetActiveDownloadRequestIDs(TSet<FString>& OutActiveIDs) const {}
+
 protected:
 	/**
 	* Designed to be called internally by AddRequest to associate our incoming request with any 
diff --git a/Engine/Source/Runtime/Online/BackgroundHTTP/Public/Interfaces/IBackgroundHttpRequest.h b/Engine/Source/Runtime/Online/BackgroundHTTP/Public/Interfaces/IBackgroundHttpRequest.h
--- a/Engine/Source/Runtime/Online/BackgroundHTTP/Public/Interfaces/IBackgroundHttpRequest.h
+++ b/Engine/Source/Runtime/Online/BackgroundHTTP/Public/Interfaces/IBackgroundHttpRequest.h
@@ -125,13 +125,6 @@ public:
 	*/
 	virtual const TArray<FString>& GetURLList() const = 0;
 
-	/**
-	* Rotates the URL list by moving the first URL to the end.
-	* Useful for retrying a request against a different CDN after a slow or failed attempt.
-	* Should only be called when the request is not being processed.
-	*/
-	virtual void RotateURLList() {}
-
 	/**
 	* Sets expected result size for the request, might be used by the implementation
 	* to check file size before invoking completion notification. 
@@ -227,6 +220,12 @@ public:
 	 */
 	virtual void CancelRequest() = 0;
 
+	/**
+	 * Re-queues a request by canceling the current platform download and creating a fresh one.
+	 * Rotates the URL list so the retry hits a different CDN. Safe to call on active requests.
+	 */
+	virtual void RequeueRequest() {}
+
     /**
     * Called to pause a request that is still being processed
     */
